Zen Oplog Utils
- Oplog manifest downloads
- Oplog diffing

diff --git a/Engine/Source/Developer/ZenOplogUtils/Private/DiffCompactBinary.cpp b/Engine/Source/Developer/ZenOplogUtils/Private/DiffCompactBinary.cpp
new file mode 100644
--- /dev/null
+++ b/Engine/Source/Developer/ZenOplogUtils/Private/DiffCompactBinary.cpp
@@ -0,0 +1,156 @@
+// Copyright Epic Games, Inc. All Rights Reserved.
+ 
+#include "Experimental/DiffCompactBinary.h"
+#include "Misc/StringBuilder.h"
+
+namespace UE
+{
+	static constexpr TCHAR GCbDelimiter = '/';
+
+	void ReportCompactBinaryDifferences(TArray<FUtf8StringView>& PropertyNameStack, FCbFieldView Entry1, FCbFieldView Entry2, TArray<FCbEntryDifference>& OutResults);
+
+	FUtf8String CreatePropertyHierarchicalName(const TArray<FUtf8StringView>& PropertyNameStack, FUtf8StringView PropertyName = {})
+	{
+		TUtf8StringBuilder<256> NameBuilder;
+		for (const FUtf8StringView& Property : PropertyNameStack)
+		{
+			NameBuilder.Append(Property);
+			NameBuilder.Add(GCbDelimiter);
+		}
+		if (PropertyName.Len() == 0 && NameBuilder.Len() > 0)	// Remove the final separator if there is no suffix 
+		{
+			NameBuilder.RemoveSuffix(1);
+		}
+		else
+		{
+			NameBuilder.Append(PropertyName);
+		}
+		return NameBuilder.ToString();
+	}
+
+	void ReportCompactBinaryDifferences(TArray<FUtf8StringView>& PropertyNameStack, FCbObjectView Entry1, FCbObjectView Entry2, TArray<FCbEntryDifference>& OutResults)
+	{
+		// Reduce the number of field lookups required by caching them once during key lookup
+		struct FCachedObjectFields
+		{
+			FCbFieldView Entry1;
+			FCbFieldView Entry2;
+		};
+		TArray<FUtf8StringView, TInlineAllocator<16>> AllKeys;			// Cache all keys found in both objects
+		TArray<FCachedObjectFields, TInlineAllocator<16>> AllFields;		// Cache the fields associated with the above keys
+		for (const FCbFieldView EntryChild : Entry1)
+		{
+			if (EntryChild.HasName())
+			{
+				AllKeys.Add(EntryChild.GetName());
+				AllFields.Add(FCachedObjectFields
+				{
+					.Entry1 = EntryChild
+				});
+			}
+		}
+		for (const FCbFieldView EntryChild : Entry2)
+		{
+			if (EntryChild.HasName())
+			{
+				int32 NewIndex = AllKeys.AddUnique(EntryChild.GetName());
+				if (NewIndex >= AllFields.Num())	// New entry needed, this field is not in object1
+				{
+					AllFields.Emplace();
+				}
+				AllFields[NewIndex].Entry2 = EntryChild;
+			}
+		}
+
+		for (int32 FieldIndex = 0; FieldIndex < AllFields.Num(); ++FieldIndex)
+		{
+			const FCbFieldView& Entry1Child = AllFields[FieldIndex].Entry1;
+			const FCbFieldView& Entry2Child = AllFields[FieldIndex].Entry2;
+			if (Entry1Child && Entry2Child)
+			{
+				ReportCompactBinaryDifferences(PropertyNameStack, Entry1Child, Entry2Child, OutResults);
+			}
+			else
+			{
+				OutResults.Add(
+				{
+					.PropertyName = CreatePropertyHierarchicalName(PropertyNameStack, AllKeys[FieldIndex]),
+					.OldValue = Entry1Child,
+					.NewValue = Entry2Child
+				});
+			}
+		}
+	}
+
+	// This is making the assumption that arrays are sorted in some fashion!
+	// Testing unsorted arrays is *much* slower
+	void ReportCompactBinaryDifferences(TArray<FUtf8StringView>& PropertyNameStack, FCbArrayView Entry1, FCbArrayView Entry2, TArray<FCbEntryDifference>& OutResults)
+	{
+		if (Entry1.Num() == Entry2.Num())
+		{
+			FCbFieldViewIterator Array1Iterator = Entry1.CreateViewIterator();
+			FCbFieldViewIterator Array2Iterator = Entry2.CreateViewIterator();
+			int32 ArrayIndex = 0;
+			UTF8CHAR ArrayIndexBuffer[32] = { 0 };
+			while (Array1Iterator)
+			{
+				FCStringUtf8::Sprintf(ArrayIndexBuffer, "%d", ArrayIndex++);
+				PropertyNameStack.Push(ArrayIndexBuffer);
+				ReportCompactBinaryDifferences(PropertyNameStack, Array1Iterator, Array2Iterator, OutResults);
+				PropertyNameStack.Pop();
+				++Array1Iterator;
+				++Array2Iterator;
+			}
+		}
+		else
+		{
+			// Arrays of different sizes, just report the array property
+			OutResults.Add(
+			{
+				.PropertyName = CreatePropertyHierarchicalName(PropertyNameStack),
+				.OldValue = Entry1.AsFieldView(),
+				.NewValue = Entry2.AsFieldView()
+			});
+		}
+	}
+
+	void ReportCompactBinaryDifferences(TArray<FUtf8StringView>& PropertyNameStack, FCbFieldView Entry1, FCbFieldView Entry2, TArray<FCbEntryDifference>& OutResults)
+	{
+		const bool EntryHasName = Entry1.HasName() && Entry2.HasName();
+		if (EntryHasName)
+		{
+			PropertyNameStack.Push(Entry1.GetName());
+		}
+		if (Entry1.IsObject() && Entry2.IsObject())
+		{
+			ReportCompactBinaryDifferences(PropertyNameStack, Entry1.AsObjectView(), Entry2.AsObjectView(), OutResults);
+		}
+		else if (Entry1.IsArray() && Entry2.IsArray())
+		{
+			ReportCompactBinaryDifferences(PropertyNameStack, Entry1.AsArrayView(), Entry2.AsArrayView(), OutResults);
+		}
+		else if (!Entry1.Equals(Entry2))
+		{
+			OutResults.Add(
+			{
+				.PropertyName = CreatePropertyHierarchicalName(PropertyNameStack),
+				.OldValue = Entry1,
+				.NewValue = Entry2
+			});
+		}
+		if (EntryHasName)
+		{
+			PropertyNameStack.Pop();
+		}
+	}
+
+	TArray<FCbEntryDifference> DiffCompactBinary(const FCbFieldView& cb1, const FCbFieldView& cb2)
+	{
+		TArray<FCbEntryDifference> Results;
+		TArray<FUtf8StringView> PropertyNames;	// Keep a stack of property names
+		PropertyNames.Reserve(10);
+		ReportCompactBinaryDifferences(PropertyNames, cb1, cb2, Results);
+		return Results;
+	}
+
+}	// namespace UE
diff --git a/Engine/Source/Developer/ZenOplogUtils/Private/ZenOplogDiff.cpp b/Engine/Source/Developer/ZenOplogUtils/Private/ZenOplogDiff.cpp
new file mode 100644
--- /dev/null
+++ b/Engine/Source/Developer/ZenOplogUtils/Private/ZenOplogDiff.cpp
@@ -0,0 +1,317 @@
+// Copyright Epic Games, Inc. All Rights Reserved.
+
+#include "Experimental/ZenOplogDiff.h"
+#include "Experimental/ZenOplogManifest.h"
+#include "Algo/Sort.h"
+#include "Async/ParallelFor.h"
+#include "Serialization/JsonSerializer.h"
+#include "Serialization/CompactBinarySerialization.h"
+#include "Serialization/CompactBinaryValue.h"
+#include "Serialization/CompactBinaryWriter.h"
+
+DEFINE_LOG_CATEGORY_STATIC(LogZenOplogUtils, Log, All);
+
+namespace UE
+{
+	bool OpDiffContainsCookedOutputChange(const UE::FOplogDiffChangedOp& Op)
+	{
+		static const FUtf8String PackageDataArrayName = TEXT("packagedata");
+		static const FUtf8String BulkDataArrayName = TEXT("bulkdata");
+		static const FUtf8String FilesArrayName = TEXT("files");
+		static const FUtf8String CookedDataHashPropertyName = TEXT("data");
+		return (OpDiffContainsArrayValueChange(Op, PackageDataArrayName, CookedDataHashPropertyName) ||
+			OpDiffContainsArrayValueChange(Op, BulkDataArrayName, CookedDataHashPropertyName) ||
+			OpDiffContainsArrayValueChange(Op, FilesArrayName, CookedDataHashPropertyName));
+	}
+
+	FOplogDiffResults DiffManifests(const FOplogManifest& Manifest1, const FOplogManifest& Manifest2)
+	{
+		FOplogDiffResults Results;
+		Results.Manifest1Data = Manifest1.AllOpsCbData;
+		Results.Manifest2Data = Manifest2.AllOpsCbData;
+
+		// Find all ops that exist in manifest 2 but not in manifest 1
+		{
+			TArray<FOplogDiffResults> TaskContexts;
+			ParallelForWithTaskContext(TaskContexts, Manifest2.Ops.Num(), [&Manifest1, &Manifest2](FOplogDiffResults& RunContext, int32 OpIndex) 
+			{
+				const FUtf8String& Op2Key = Manifest2.Ops[OpIndex].GetKey();
+				const int32* FoundManifest1Index = Manifest1.OpKeyToIndex.Find(Op2Key);
+				if (FoundManifest1Index == nullptr)
+				{
+					RunContext.OpsMissingInManifest1.Add(Op2Key);
+				}
+			});
+			for (FOplogDiffResults& TaskResult : TaskContexts)	// Combine results
+			{
+				Results.OpsMissingInManifest1.Append(MoveTemp(TaskResult.OpsMissingInManifest1));
+			}
+		}
+
+		// Diff all ops in manifest1 against manifest 2
+		{
+			TArray<FOplogDiffResults> TaskContexts;
+			ParallelForWithTaskContext(TaskContexts, Manifest1.Ops.Num(), [&Manifest1, &Manifest2](FOplogDiffResults& RunContext, int32 OpIndex) 
+			{
+				const TPair<FUtf8String, uint32>& Op1KeyAndHash = Manifest1.OpKeysAndHashes[OpIndex];
+				const int32* FoundManifest2Index = Manifest2.OpKeyToIndex.FindByHash(Op1KeyAndHash.Value, Op1KeyAndHash.Key);
+				if (FoundManifest2Index == nullptr)
+				{
+					RunContext.OpsMissingInManifest2.Add(Op1KeyAndHash.Key);
+				}
+				else
+				{
+					const FOplogManifest::FOp& Op2 = Manifest2.Ops[*FoundManifest2Index];
+					TArray<FCbEntryDifference> Differences = DiffCompactBinary(Manifest1.Ops[OpIndex].CbData, Op2.CbData);
+					if (Differences.Num() == 0)
+					{
+						RunContext.IdenticalOps.Add(Op1KeyAndHash.Key);
+					}
+					else
+					{
+						Algo::SortBy(Differences, &FCbEntryDifference::PropertyName);	// Sort changed properties by name
+						FOplogDiffChangedOp ChangedOp = { Op1KeyAndHash.Key , MoveTemp(Differences) };
+						if (OpDiffContainsCookedOutputChange(ChangedOp))
+						{
+							RunContext.ChangedOpsWithOutputDifferences.Emplace(MoveTemp(ChangedOp));
+						} 
+						else
+						{
+							RunContext.ChangedOpsWithSameOutput.Emplace(MoveTemp(ChangedOp));
+						}
+					}
+				}
+			});
+			for (FOplogDiffResults& TaskResult : TaskContexts)	// Combine results
+			{
+				Results.ChangedOpsWithOutputDifferences.Append(MoveTemp(TaskResult.ChangedOpsWithOutputDifferences));
+				Results.ChangedOpsWithSameOutput.Append(MoveTemp(TaskResult.ChangedOpsWithSameOutput));
+				Results.IdenticalOps.Append(MoveTemp(TaskResult.IdenticalOps));
+				Results.OpsMissingInManifest2.Append(MoveTemp(TaskResult.OpsMissingInManifest2));
+			}
+		}
+
+		return Results;
+	}
+
+	bool ChangeIsArrayOrEntryModification(const FCbEntryDifference& Entry, FUtf8StringView ArrayName, FUtf8StringView PropertyName)
+	{
+		check(ArrayName.Len() > 0 && PropertyName.Len() > 0);
+		FUtf8StringView EntryPropName(Entry.PropertyName);
+		if (EntryPropName.StartsWith(ArrayName, ESearchCase::IgnoreCase))
+		{
+			EntryPropName.RightChopInline(ArrayName.Len());
+			if (EntryPropName.Len() == 0)
+			{
+				return true;	// Entry contains the array name only
+			}
+			if (EntryPropName.EndsWith(PropertyName, ESearchCase::IgnoreCase))
+			{
+				EntryPropName.LeftChopInline(PropertyName.Len());
+				if (EntryPropName.StartsWith('/') && EntryPropName.EndsWith('/') && EntryPropName.Len() > 2)
+				{
+					EntryPropName.MidInline(1, EntryPropName.Len() - 2);
+					// Ensure anything remaining parses to an integer array index
+					auto RemainingAsTchar = StringCast<TCHAR>(EntryPropName.GetData(), EntryPropName.Len());
+					TCHAR* TextEnd = nullptr;
+					FCString::Strtoi(RemainingAsTchar.Get(), &TextEnd, 10);
+					if (TextEnd == (RemainingAsTchar.Get() + RemainingAsTchar.Length()))
+					{
+						return true;
+					}
+				}
+			}
+		}
+		return false;
+	}
+
+	bool OpDiffContainsArrayValueChange(const FOplogDiffChangedOp& Op, FUtf8StringView ArrayName, FUtf8StringView PropertyName)
+	{
+		for (const FCbEntryDifference& Entry : Op.Value)
+		{
+			if (ChangeIsArrayOrEntryModification(Entry, ArrayName, PropertyName))
+			{
+				return true;
+			}
+		}
+		return false;
+	}
+
+	void SortDiffResults(FOplogDiffResults& FullResults)
+	{
+		Algo::Sort(FullResults.IdenticalOps);
+		Algo::Sort(FullResults.OpsMissingInManifest1);
+		Algo::Sort(FullResults.OpsMissingInManifest2);
+		Algo::SortBy(FullResults.ChangedOpsWithSameOutput, &FOplogDiffChangedOp::Key);
+		Algo::SortBy(FullResults.ChangedOpsWithOutputDifferences, &FOplogDiffChangedOp::Key);
+	}
+
+	// Helper for outputting ops
+	template<typename OpEntry>
+	void OutputEachOp(const TArray<OpEntry>& Results, TFunction<void(const OpEntry&)> OutputFn)
+	{
+		for (const OpEntry& Package : Results)
+		{
+			OutputFn(Package);
+		}
+	}
+
+	template<bool bOutputPropertyDifferences>
+	void OutputManifestDiffResultsToJsonImpl(TSharedRef<TJsonWriter<>>& JsonWriter, const FOplogDiffResults& Results, bool bOutputIdentical)
+	{
+		auto OutputChangedOp = [&JsonWriter](const FOplogDiffChangedOp& Package)
+		{
+			JsonWriter->WriteObjectStart();
+			JsonWriter->WriteValue("Key", Package.Key);
+			JsonWriter->WriteArrayStart("Diffs");
+			TUtf8StringBuilder<1024> PropertyDataAsJson;
+			for (const FCbEntryDifference& Difference : Package.Value)
+			{
+				JsonWriter->WriteObjectStart();
+				JsonWriter->WriteValue("Property", Difference.PropertyName);
+				if constexpr (bOutputPropertyDifferences)
+				{
+					if (Difference.OldValue.HasValue())
+					{
+						PropertyDataAsJson.Reset();
+						CompactBinaryToJson(Difference.OldValue, PropertyDataAsJson);
+						JsonWriter->WriteArrayStart("Old");	// We must surround values in arrays since objects cannot contain un-named values
+						JsonWriter->WriteRawJSONValue(PropertyDataAsJson.ToString());
+						JsonWriter->WriteArrayEnd();
+					}
+					if (Difference.NewValue.HasValue())
+					{
+						PropertyDataAsJson.Reset();
+						CompactBinaryToJson(Difference.NewValue, PropertyDataAsJson);
+						JsonWriter->WriteArrayStart("New");
+						JsonWriter->WriteRawJSONValue(PropertyDataAsJson.ToString());
+						JsonWriter->WriteArrayEnd();
+					}
+				}
+				JsonWriter->WriteObjectEnd();
+			}
+			JsonWriter->WriteArrayEnd();	// Diffs
+			JsonWriter->WriteObjectEnd();
+		};
+		JsonWriter->WriteObjectStart("Diff");
+		JsonWriter->WriteArrayStart("ChangedOpsWithSameOutput");
+		OutputEachOp<FOplogDiffChangedOp>(Results.ChangedOpsWithSameOutput, OutputChangedOp);
+		JsonWriter->WriteArrayEnd();	// ChangedOpsSameOutput
+		JsonWriter->WriteArrayStart("ChangedOpsWithOutputDifferences");
+		OutputEachOp<FOplogDiffChangedOp>(Results.ChangedOpsWithOutputDifferences, OutputChangedOp);
+		JsonWriter->WriteArrayEnd();	// ChangedOpsWithOutputDifferences
+		JsonWriter->WriteArrayStart("MissingInManifest1");
+		OutputEachOp<FUtf8String>(Results.OpsMissingInManifest1, [&JsonWriter](const FUtf8String& Package)
+		{
+			JsonWriter->WriteValue(Package);
+		});
+		JsonWriter->WriteArrayEnd();	// MissingInManifest1
+		JsonWriter->WriteArrayStart("MissingInManifest2");
+		OutputEachOp<FUtf8String>(Results.OpsMissingInManifest2, [&JsonWriter](const FUtf8String& Package)
+		{
+			JsonWriter->WriteValue(Package);
+		});
+		JsonWriter->WriteArrayEnd();	// MissingInManifest2
+		if (bOutputIdentical)
+		{
+			JsonWriter->WriteArrayStart("Identical");
+			OutputEachOp<FUtf8String>(Results.IdenticalOps, [&JsonWriter](const FUtf8String& Package)
+			{
+				JsonWriter->WriteValue(Package);
+			});
+			JsonWriter->WriteArrayEnd();	// Identical
+		}
+		JsonWriter->WriteObjectEnd();
+	}
+
+	template<bool bOutputPropertyDifferences>
+	void OutputManifestDiffResultsToCompactBinaryImpl(FCbWriter& Writer, const FOplogDiffResults& Results, bool bOutputIdentical)
+	{
+		auto OutputChangedOp = [&Writer](const FOplogDiffChangedOp& Package)
+		{
+			Writer.BeginObject();
+			Writer.AddString("Key", Package.Key);
+			Writer.BeginArray("Diffs");
+			for (const FCbEntryDifference& Difference : Package.Value)
+			{
+				Writer.BeginObject();
+				Writer.AddString("Property", Difference.PropertyName);
+				if constexpr (bOutputPropertyDifferences)
+				{
+					if (Difference.OldValue.HasValue())
+					{
+						Writer.BeginArray("Old");
+						Writer.AddField(Difference.OldValue);
+						Writer.EndArray();
+					}
+					if (Difference.NewValue.HasValue())
+					{
+						Writer.BeginArray("New");
+						Writer.AddField(Difference.NewValue);
+						Writer.EndArray();
+					}
+				}
+				Writer.EndObject();
+			}
+			Writer.EndArray();	// diffs
+			Writer.EndObject();
+		};
+		Writer.BeginObject("Diff");
+		Writer.BeginArray("ChangedOpsWithSameOutput");
+		OutputEachOp<FOplogDiffChangedOp>(Results.ChangedOpsWithSameOutput, OutputChangedOp);
+		Writer.EndArray();	// ChangedOpsSameOutput
+		Writer.BeginArray("ChangedOpsWithOutputDifferences");
+		OutputEachOp<FOplogDiffChangedOp>(Results.ChangedOpsWithOutputDifferences, OutputChangedOp);
+		Writer.EndArray();	// ChangedOpsWithOutputDifferences
+		Writer.BeginArray("MissingInManifest1");
+		OutputEachOp<FUtf8String>(Results.OpsMissingInManifest1, [&Writer](const FUtf8String& Package)
+		{
+			Writer.AddString(Package);
+		});
+		Writer.EndArray();	// MissingInManifest1
+		Writer.BeginArray("MissingInManifest2");
+		OutputEachOp<FUtf8String>(Results.OpsMissingInManifest2, [&Writer](const FUtf8String& Package)
+		{
+			Writer.AddString(Package);
+		});
+		Writer.EndArray();	// MissingInManifest2
+		if (bOutputIdentical)
+		{
+			Writer.BeginArray("Identical");
+			OutputEachOp<FUtf8String>(Results.IdenticalOps, [&Writer](const FUtf8String& Package)
+			{
+				Writer.AddString(Package);
+			});
+			Writer.EndArray();	// Identical
+		}
+		
+		Writer.EndObject();
+	}
+
+	void OutputManifestDiffResultsToJson(TSharedRef<TJsonWriter<>>& Writer, const FOplogDiffResults& Results, EOutputManifestDiffOptions Flags)
+	{
+		const bool bOutputIdentical = EnumHasAnyFlags(Flags, EOutputManifestDiffOptions::OutputIdenticalOps);
+		if (EnumHasAnyFlags(Flags, EOutputManifestDiffOptions::OutputDifferences))
+		{
+			OutputManifestDiffResultsToJsonImpl<true>(Writer, Results, bOutputIdentical);
+		}
+		else
+		{
+			OutputManifestDiffResultsToJsonImpl<false>(Writer, Results, bOutputIdentical);
+		}
+	}
+
+	void OutputManifestDiffResultsToCompactBinary(FCbWriter& Writer, const FOplogDiffResults& Results, EOutputManifestDiffOptions Flags)
+	{
+		const bool bOutputIdentical = EnumHasAnyFlags(Flags, EOutputManifestDiffOptions::OutputIdenticalOps);
+		if (EnumHasAnyFlags(Flags, EOutputManifestDiffOptions::OutputDifferences))
+		{
+			return OutputManifestDiffResultsToCompactBinaryImpl<true>(Writer, Results, bOutputIdentical);
+		}
+		else
+		{
+			return OutputManifestDiffResultsToCompactBinaryImpl<false>(Writer, Results, bOutputIdentical);
+		}
+	}
+
+}	// namespace UE
diff --git a/Engine/Source/Developer/ZenOplogUtils/Private/ZenOplogManifest.cpp b/Engine/Source/Developer/ZenOplogUtils/Private/ZenOplogManifest.cpp
new file mode 100644
--- /dev/null
+++ b/Engine/Source/Developer/ZenOplogUtils/Private/ZenOplogManifest.cpp
@@ -0,0 +1,220 @@
+// Copyright Epic Games, Inc. All Rights Reserved.
+
+#include "Experimental/ZenOplogManifest.h"
+#include "Experimental/DiffCompactBinary.h"
+#include "Async/ParallelFor.h"
+#include "HAL/FileManager.h"
+#include "Misc/Paths.h"
+#include "Serialization/CompactBinary.h"
+#include "Serialization/CompactBinarySerialization.h"
+
+namespace UE
+{
+
+	FUtf8String FOplogManifest::FOp::GetKey() const
+	{
+		return FUtf8String(CbData["key"].AsString());
+	}
+
+	FIoHash FOplogManifest::FOp::GetArtifactHash() const
+	{
+		return CbData["meta.cook.artifacts"].AsBinaryAttachment();
+	}
+
+	FIoHash FOplogManifest::FOp::CookImportExport() const
+	{
+		return CbData["meta.cook.importexport"].AsBinaryAttachment();
+	}
+
+	FIoHash FOplogManifest::FOp::GetLogs() const
+	{
+		return CbData["meta.cook.logs"].AsBinaryAttachment();
+	}
+
+	FOplogManifest::FOp::FPackageStoreEntry FOplogManifest::FOp::GetPackageStoreEntry() const
+	{
+		FOplogManifest::FOp::FPackageStoreEntry Data;
+		if (FCbObjectView PackageStoreEntry = CbData["packagestoreentry"].AsObjectView())
+		{
+			Data.Flags = PackageStoreEntry["flags"].AsUInt32();
+			Data.PackageName = PackageStoreEntry["packagename"].AsString();
+			if (PackageStoreEntry["importedpackageids"])
+			{
+				for (FCbFieldView ArrayField : PackageStoreEntry["importedpackageids"])
+				{
+					Data.ImportedPackages.Add(ArrayField.AsUInt64());
+				}
+			}
+			if (PackageStoreEntry["softpackagereferences"])
+			{
+				for (FCbFieldView ArrayField : PackageStoreEntry["softpackagereferences"])
+				{
+					Data.SoftPackageReferences.Add(ArrayField.AsUInt64());
+				}
+			}
+		}
+		return Data;
+	}
+
+	TArray<FOplogManifest::FOp::FPackageData> FOplogManifest::FOp::GetPackageDatas() const
+	{
+		TArray<FOplogManifest::FOp::FPackageData> ParsedData;
+		FCbArrayView PackageDataArray = CbData["packagedata"].AsArrayView();
+		for (FCbFieldView PackageDataEntry : PackageDataArray)
+		{
+			FCbObjectView PackageDataObject = PackageDataEntry.AsObjectView();
+			FOplogManifest::FOp::FPackageData NewPackageData;
+			NewPackageData.ID = PackageDataEntry["id"].AsObjectId();
+			NewPackageData.Size = PackageDataEntry["size"].AsInt64();
+			NewPackageData.RawSize = PackageDataEntry["rawsize"].AsInt64();
+			NewPackageData.DataHash = PackageDataEntry["data"].AsBinaryAttachment();
+			NewPackageData.Filename = PackageDataEntry["filename"].AsString();
+			ParsedData.Add(NewPackageData);
+		}
+		return ParsedData;
+	}
+
+	TArray<FOplogManifest::FOp::FBulkData> FOplogManifest::FOp::GetBulkDatas() const
+	{
+		TArray<FOplogManifest::FOp::FBulkData> ParsedData;
+		FCbArrayView BulkDataArray = CbData["bulkdata"].AsArrayView();
+		for (FCbFieldView BulkDataEntry : BulkDataArray)
+		{
+			FCbObjectView BulkDataObject = BulkDataEntry.AsObjectView();
+			FOplogManifest::FOp::FBulkData NewBulkData;
+			NewBulkData.ID = BulkDataEntry["id"].AsObjectId();
+			NewBulkData.TypeStr = BulkDataEntry["type"].AsString();
+			NewBulkData.Size = BulkDataEntry["size"].AsInt64();
+			NewBulkData.RawSize = BulkDataEntry["rawsize"].AsInt64();
+			NewBulkData.DataHash = BulkDataEntry["data"].AsBinaryAttachment();
+			NewBulkData.Filename = BulkDataEntry["filename"].AsString();
+			ParsedData.Add(NewBulkData);
+		}
+		return ParsedData;
+	}
+
+	TArray<FOplogManifest::FOp::FFile> FOplogManifest::FOp::GetFiles() const
+	{
+		TArray<FOplogManifest::FOp::FFile> ParsedData;
+		FCbArrayView FilesArray = CbData["files"].AsArrayView();
+		for (FCbFieldView FileEntry : FilesArray)
+		{
+			FOplogManifest::FOp::FFile NewFile;
+			NewFile.ID = FileEntry["id"].AsObjectId();
+			NewFile.ServerPath = FileEntry["serverpath"].AsString();
+			NewFile.ClientPath = FileEntry["clientpath"].AsString();
+			NewFile.DataHash = FileEntry["data"].AsBinaryAttachment();
+			ParsedData.Add(NewFile);
+		}
+		return ParsedData;
+	}
+
+	TArray<FOplogManifest::FOp::FMeta> FOplogManifest::FOp::GetMetas() const
+	{
+		TArray<FOplogManifest::FOp::FMeta> ParsedData;
+		FCbArrayView MetaArray = CbData["meta"].AsArrayView();
+		for (FCbFieldView MetaEntry : MetaArray)
+		{
+			FOplogManifest::FOp::FMeta NewMeta;
+			NewMeta.ID = MetaEntry["id"].AsObjectId();
+			NewMeta.Name = MetaEntry["name"].AsString();
+			NewMeta.DataHash = MetaEntry["data"].AsBinaryAttachment();
+			ParsedData.Add(NewMeta);
+		}
+		return ParsedData;
+	}
+
+	bool FOplogManifest::RebuildCachedData()
+	{
+		check(OpKeyToIndex.Num() == 0);
+		check(OpKeysAndHashes.Num() == 0);
+		OpKeyToIndex.Reserve(Ops.Num());
+		OpKeysAndHashes.AddDefaulted(Ops.Num());
+
+		// Collect the key hashes and key -> index mappings on task threads
+		struct CollectedOpKeyToIndex
+		{
+			FUtf8String Key;
+			uint32 KeyHash;
+			int32 OpIndex;
+		};
+		TArray<TArray<CollectedOpKeyToIndex>> PerTaskOpKeyToIndex;
+		ParallelForWithTaskContext(PerTaskOpKeyToIndex, Ops.Num(), [this](TArray<CollectedOpKeyToIndex>& ThisTaskOpKeyToIndex, int32 OpIndex) 
+		{
+			FUtf8String ThisKey = Ops[OpIndex].GetKey();
+			uint32 ThisKeyHash = GetTypeHash(ThisKey);
+			ThisTaskOpKeyToIndex.Emplace(ThisKey, ThisKeyHash, OpIndex);
+			OpKeysAndHashes[OpIndex] = { MoveTemp(ThisKey), ThisKeyHash };
+		});
+
+		// Combine results
+		bool bDuplicateKeyFound = false;
+		for (TArray<CollectedOpKeyToIndex>& PerTaskEntries : PerTaskOpKeyToIndex)
+		{
+			for (CollectedOpKeyToIndex& Entry : PerTaskEntries)
+			{
+				bDuplicateKeyFound |= (OpKeyToIndex.FindOrAddByHash(Entry.KeyHash, MoveTemp(Entry.Key), Entry.OpIndex) != Entry.OpIndex);
+			}
+		}
+		return !bDuplicateKeyFound;
+	}
+
+	FLoadOplogManifestResult::FLoadOplogManifestResult(FLoadOplogManifestResult::EStatus Stat, TOptional<FOplogManifest>&& Man)
+		: Result(Stat)
+		, Manifest(MoveTemp(Man))
+	{
+	}
+
+	FLoadOplogManifestResult LoadOplogManifestFromCompactBinary(const FCbObject& CompactBinary)
+	{
+		FOplogManifest LoadedManifest;
+
+		FCbField AllOps = CompactBinary["ops"];
+		if (!AllOps.IsArray())
+		{
+			return FLoadOplogManifestResult(FLoadOplogManifestResult::EStatus::ErrorMalformedCompactBinary);
+		}
+
+		LoadedManifest.AllOpsCbData = MoveTemp(AllOps);
+		for (FCbFieldView OpEntry : LoadedManifest.AllOpsCbData)
+		{
+			if (!OpEntry["key"].IsString())	// We expect each op entry to contain a key at minimum
+			{
+				return FLoadOplogManifestResult(FLoadOplogManifestResult::EStatus::ErrorMalformedCompactBinary);
+			}
+
+			FOplogManifest::FOp NewOpData;
+			NewOpData.CbData = MoveTemp(OpEntry);
+			LoadedManifest.Ops.Add(MoveTemp(NewOpData));
+		}
+		if (!LoadedManifest.RebuildCachedData())
+		{
+			return FLoadOplogManifestResult(FLoadOplogManifestResult::EStatus::Error);
+		}
+		return FLoadOplogManifestResult(FLoadOplogManifestResult::EStatus::Ok, MoveTemp(LoadedManifest));
+	}
+
+	FLoadOplogManifestResult LoadOplogManifestFromFile(FString Path)
+	{
+		FString FileExtension = FPaths::GetExtension(Path).ToLower();
+		if (FileExtension == TEXT("cb"))
+		{
+			TUniquePtr<FArchive> Ar(IFileManager::Get().CreateFileReader(*Path));
+			if (Ar)
+			{
+				FCbObject RootObject = LoadCompactBinary(*Ar).AsObject();
+				if (!RootObject)
+				{
+					return FLoadOplogManifestResult(FLoadOplogManifestResult::EStatus::ErrorMalformedCompactBinary);
+				}
+				return LoadOplogManifestFromCompactBinary(RootObject);
+			}
+			else
+			{
+				return FLoadOplogManifestResult(FLoadOplogManifestResult::EStatus::ErrorFileNotFound);
+			}
+		}
+		return FLoadOplogManifestResult(FLoadOplogManifestResult::EStatus::ErrorUnsupportedFormat);
+	}
+
+}	// namespace UE
diff --git a/Engine/Source/Developer/ZenOplogUtils/Private/ZenOplogSnapshots.cpp b/Engine/Source/Developer/ZenOplogUtils/Private/ZenOplogSnapshots.cpp
new file mode 100644
--- /dev/null
+++ b/Engine/Source/Developer/ZenOplogUtils/Private/ZenOplogSnapshots.cpp
@@ -0,0 +1,253 @@
+// Copyright Epic Games, Inc. All Rights Reserved.
+
+#include "Experimental/ZenOplogSnapshots.h"
+#include "Experimental/BuildServerInterface.h"
+#include "Experimental/ZenServerInterface.h"
+#include "DesktopPlatformModule.h"
+#include "Misc/App.h"
+#include "Misc/FileHelper.h"
+#include "Misc/Paths.h"
+#include "Serialization/JsonSerializer.h"
+
+DEFINE_LOG_CATEGORY_STATIC(LogOplogSnapshotUtils, Log, All);
+
+namespace UE
+{
+	int32 RunZenExeSync(const TCHAR* CommandLine)
+	{
+		int32 ReturnCode = -1;
+		FString ZenUtilPath = FPaths::ConvertRelativePathToFull(UE::Zen::GetLocalInstallUtilityPath());
+
+		FProcHandle ProcHandle = FPlatformProcess::CreateProc(*ZenUtilPath, CommandLine, false, false, false, nullptr, 0, nullptr, nullptr);
+		if (ProcHandle.IsValid())
+		{
+			FPlatformProcess::WaitForProc(ProcHandle);
+			if (!FPlatformProcess::GetProcReturnCode(ProcHandle, &ReturnCode))
+			{
+				UE_LOG(LogOplogSnapshotUtils, Warning, TEXT("Failed to get return code from zen process"));
+			}
+			FPlatformProcess::CloseProc(ProcHandle);
+		}
+		else
+		{
+			UE_LOG(LogOplogSnapshotUtils, Warning, TEXT("Failed to launch zen utility from %s"), *ZenUtilPath);
+		}
+		return ReturnCode;
+	}
+
+	FDownloadOplogManifestResult DownloadOplogManifest(const FZenSnapshotDescriptor& Descriptor, FString OutputDirectory, bool bDownloadAsJson)
+	{
+		FString OidcTokenExePath = FPaths::ConvertRelativePathToFull(FDesktopPlatformModule::Get()->GetOidcTokenExecutableFilename(FPaths::RootDir()));
+		FString OutputDir = FPaths::ConvertRelativePathToFull(FPaths::Combine(FDesktopPlatformModule::Get()->GetUserTempPath(), TEXT("OplogManifests")));
+		FString ManifestFileType = bDownloadAsJson ? TEXT("json") : TEXT("cb");
+		if (!OutputDirectory.IsEmpty() && FPaths::DirectoryExists(OutputDirectory))
+		{
+			OutputDir = OutputDirectory;
+		}
+
+		FString ManifestPath = FPaths::Combine(OutputDir, FPaths::MakeValidFileName(FString::Printf(TEXT("%s-%s-%s.%s"), *Descriptor.Namespace, *Descriptor.Bucket, *Descriptor.BuildID, *ManifestFileType)));
+		
+		TStringBuilder<1024> CommandLineArgs;
+		CommandLineArgs.Append(TEXT("oplog-download"));
+		CommandLineArgs.Append(FString::Printf(TEXT(" --host \"%s\""), *Descriptor.CloudHost));
+		CommandLineArgs.Append(FString::Printf(TEXT(" --namespace \"%s\""), *Descriptor.Namespace));
+		CommandLineArgs.Append(FString::Printf(TEXT(" --bucket \"%s\""), *Descriptor.Bucket));
+		CommandLineArgs.Append(FString::Printf(TEXT(" --build-id \"%s\""), *Descriptor.BuildID));
+		CommandLineArgs.Append(FString::Printf(TEXT(" --oidctoken-exe-path \"%s\""), *OidcTokenExePath));
+		CommandLineArgs.Append(FString::Printf(TEXT(" --output-path \"%s\""), *ManifestPath));
+		if (FApp::IsUnattended())
+		{
+			CommandLineArgs.Append(TEXT(" --oidctoken-exe-unattended"));
+		}
+
+		int32 ZenResult = RunZenExeSync(*CommandLineArgs);
+		return FDownloadOplogManifestResult
+		{
+			.Result = (ZenResult == 0) ? FDownloadOplogManifestResult::EStatus::Ok : FDownloadOplogManifestResult::EStatus::Error,
+			.DownloadedFilename = (ZenResult == 0) ? FString(ManifestPath) : FString()
+		};
+	}
+
+	TArray<FZenSnapshotDescriptor> ParseSnapshotDescriptorFromJson(FStringView PathToJson)
+	{
+		TArray<FZenSnapshotDescriptor> NewSnapshots;
+		FString JsonFileContent;
+		if (FFileHelper::LoadFileToString(JsonFileContent, *WriteToString<256>(PathToJson)))
+		{
+			TSharedPtr<FJsonObject> TopLevelObject;
+			const TSharedRef<TJsonReader<>>& Reader = TJsonReaderFactory<>::Create(JsonFileContent);
+			if (FJsonSerializer::Deserialize(Reader, TopLevelObject))
+			{
+				if (const TSharedPtr <FJsonValue>* SnapshotValues = TopLevelObject->Values.Find("snapshots"))
+				{
+					for (const TSharedPtr<FJsonValue>& SnapshotEntry : (*SnapshotValues)->AsArray())
+					{
+						if (const TSharedPtr<FJsonObject> SnapshotJson = SnapshotEntry->AsObject())
+						{
+							NewSnapshots.Add(
+							{
+								.CloudHost = SnapshotJson->GetStringField(TEXT("host")),
+								.Namespace = SnapshotJson->GetStringField(TEXT("namespace")),
+								.Bucket = SnapshotJson->GetStringField(TEXT("bucket")),
+								.BuildID = SnapshotJson->GetStringField(TEXT("builds-id"))
+							});
+						}
+					}
+				}
+			}
+			else
+			{
+				UE_LOG(LogOplogSnapshotUtils, Warning, TEXT("Failed to parse json in %s"), *FString(PathToJson));
+			}
+		}
+		else
+		{
+			UE_LOG(LogOplogSnapshotUtils, Warning, TEXT("Failed to load snapshot file %s"), *FString(PathToJson));
+		}
+		return NewSnapshots;
+	}
+
+	// Data used during the async bits of FindSnapshotDescriptorForBuild
+	struct FFindSnapshotDescriptorContext
+	{
+		// The project/stream/platform/cl/build to search for
+		FString FindProjectName;
+		FString FindStreamName;
+		FString FindPlatformName;
+		FString FindBuildType;
+		FString FindCommitID;
+
+		bool bLogAllOutput;		// Whether to log all build info found
+
+		// The hostname/namespace/bucket found during FindSnapshotDescriptorForBuild
+		FString FoundHostname;
+		FString FoundNamespace;
+		FString FoundBucket;
+
+		TSharedPtr<UE::Zen::Build::FBuildServiceInstance> BuildService;
+		TPromise<TOptional<FZenSnapshotDescriptor>> ResultPromise;
+	};
+
+	// Called when the build service has listed a number of builds matching the search in FFindSnapshotDescriptorContext
+	void OnBuildsFound(TSharedPtr<FFindSnapshotDescriptorContext> Context, TArray<UE::Zen::Build::FBuildServiceInstance::FBuildRecord>&& Builds)
+	{
+		for (const UE::Zen::Build::FBuildServiceInstance::FBuildRecord& BuildRecord : Builds)
+		{
+			FString CommitID = BuildRecord.GetCommitIdentifier();
+			if (Context->bLogAllOutput)
+			{
+				UE_LOG(LogOplogSnapshotUtils, Display, TEXT("\t%s"), *CommitID);
+			}
+			if (CommitID == Context->FindCommitID)
+			{
+				FZenSnapshotDescriptor FoundSnapshot;
+				FoundSnapshot.Namespace = Context->FoundNamespace;
+				FoundSnapshot.CloudHost = Context->FoundHostname;
+				FoundSnapshot.BuildID = WriteToString<64>(BuildRecord.BuildId);
+				FoundSnapshot.Bucket = Context->FoundBucket;
+				Context->ResultPromise.SetValue(FoundSnapshot);
+				return;
+			}
+		}
+		UE_LOG(LogOplogSnapshotUtils, Warning, TEXT("Failed to find a matching build with commit ID %s in %s/%s"), *Context->FindCommitID, *Context->FoundNamespace, *Context->FoundBucket);
+		Context->ResultPromise.SetValue({});
+	}
+
+	// Called when the build service has refreshed the list of all available namespaces/buckets
+	void OnNamespacesAndBucketsRefreshed(TSharedPtr<FFindSnapshotDescriptorContext> Context)
+	{
+		TMultiMap<FString, FString> NamespacesAndBuckets = Context->BuildService->GetNamespacesAndBuckets();
+		TStringBuilder<256> BucketNameToFind;	// bucket names are expected in the form of: projectname.buildtype.stream.platform
+		BucketNameToFind << Context->FindProjectName.ToLower() << TEXT(".");
+		BucketNameToFind << Context->FindBuildType.ToLower() << TEXT(".");
+		BucketNameToFind << Context->FindStreamName.ToLower() << TEXT(".");
+		BucketNameToFind << Context->FindPlatformName.ToLower();
+
+		if (Context->bLogAllOutput)
+		{
+			UE_LOG(LogOplogSnapshotUtils, Display, TEXT("Searching for bucket %s"), *BucketNameToFind);
+		}
+
+		for (const TPair<FString, FString>& NamespaceAndBucket : NamespacesAndBuckets)
+		{
+			if (Context->bLogAllOutput)
+			{
+				UE_LOG(LogOplogSnapshotUtils, Display, TEXT("%s -> %s"), *NamespaceAndBucket.Key, *NamespaceAndBucket.Value);
+			}
+			if (NamespaceAndBucket.Value == BucketNameToFind.ToString())
+			{
+				// Now query for builds in the bucket
+				Context->FoundNamespace = NamespaceAndBucket.Key;
+				Context->FoundBucket = NamespaceAndBucket.Value;
+				Context->BuildService->ListBuilds(NamespaceAndBucket.Key, NamespaceAndBucket.Value, FCbObject{},
+					[Context](TArray<UE::Zen::Build::FBuildServiceInstance::FBuildRecord>&& BuildRecords)
+				{
+					OnBuildsFound(Context, MoveTemp(BuildRecords));
+				});
+				return;
+			}
+		}
+
+		UE_LOG(LogOplogSnapshotUtils, Warning, TEXT("Failed to find a matching build namespace/bucket"));
+		Context->ResultPromise.SetValue({});
+	}
+
+	// Called when the build service successfully connects
+	void OnBuildServiceConnected(TSharedPtr<FFindSnapshotDescriptorContext> Context)
+	{
+		Context->FoundHostname = FString(Context->BuildService->GetEffectiveDomain());
+		Context->BuildService->RefreshNamespacesAndBuckets([Context]()
+		{
+			ExecuteOnGameThread(UE_SOURCE_LOCATION, [Context]()
+			{
+				OnNamespacesAndBucketsRefreshed(Context);
+			});
+		});
+	}
+
+	TFuture<TOptional<FZenSnapshotDescriptor>> FindSnapshotDescriptorForBuild(FStringView ProjectName, FStringView StreamName, FStringView PlatformName, FStringView CommitID, bool bLogAllBuilds)
+	{
+		using namespace UE::Zen::Build;
+
+		TSharedPtr<FFindSnapshotDescriptorContext> Context = MakeShared<FFindSnapshotDescriptorContext>();
+		Context->BuildService = MakeShared<FBuildServiceInstance>();
+		Context->FindProjectName = FString(ProjectName).ToLower();
+		Context->FindStreamName = FString(StreamName).ToLower();
+		Context->FindPlatformName = FString(PlatformName).ToLower();
+		Context->FindCommitID = CommitID;
+		Context->FindBuildType = "oplog";
+		Context->bLogAllOutput = bLogAllBuilds;
+
+		TFuture<TOptional<FZenSnapshotDescriptor>> ResultFuture = Context->ResultPromise.GetFuture();
+		Context->BuildService->Connect(!FApp::IsUnattended(), [Context]
+		(FBuildServiceInstance::EConnectionState ConnectionState,
+			FBuildServiceInstance::EConnectionFailureReason FailureReason)
+		{
+			if (ConnectionState == UE::Zen::Build::FBuildServiceInstance::EConnectionState::ConnectionSucceeded)
+			{
+				OnBuildServiceConnected(Context);
+			}
+			else
+			{
+				UE_LOG(LogOplogSnapshotUtils, Error, TEXT("Failed to connect to build service"));
+				switch (FailureReason)
+				{
+				case FBuildServiceInstance::EConnectionFailureReason::FailedToAcquireAccessToken:
+					UE_LOG(LogOplogSnapshotUtils, Error, TEXT("\tFailed to acquire access token"));
+					break;
+				case FBuildServiceInstance::EConnectionFailureReason::FailedToResolveHost:
+					UE_LOG(LogOplogSnapshotUtils, Error, TEXT("Failed to resolve host"));
+					break;
+				case FBuildServiceInstance::EConnectionFailureReason::UnexpectedState:
+					UE_LOG(LogOplogSnapshotUtils, Error, TEXT("Unexpected state"));
+					break;
+				default:
+					UE_LOG(LogOplogSnapshotUtils, Error, TEXT("Unknown error"));
+					break;
+				}
+				Context->ResultPromise.SetValue({});
+			}
+		});
+		return ResultFuture;
+	}
+}
\ No newline at end of file
diff --git a/Engine/Source/Developer/ZenOplogUtils/Public/Experimental/DiffCompactBinary.h b/Engine/Source/Developer/ZenOplogUtils/Public/Experimental/DiffCompactBinary.h
new file mode 100644
--- /dev/null
+++ b/Engine/Source/Developer/ZenOplogUtils/Public/Experimental/DiffCompactBinary.h
@@ -0,0 +1,33 @@
+// Copyright Epic Games, Inc. All Rights Reserved.
+
+#pragma once
+
+#include "CoreTypes.h"
+#include "Containers/Utf8String.h"
+#include "Serialization/CompactBinary.h"
+
+#define UE_API ZENOPLOGUTILS_API
+
+namespace UE
+{
+
+	// This represents a single value difference in compact binary data
+	//	Differences are reported with hierarchical names, so they can be found/reconstructed from the compact binary
+	//		For example, compact binary containing an object named packagedata with differing filename values will return a PropertyName of 'packagedata/filename'
+	//	Arrays of different sizes are reported as a single diff entry with the name of the array
+	//	Arrays of equal size with different elements will have each element diff reported separately, with the index appended to the array property name
+	//		e.g. 'packagestoreentry/importedpackages/2' or 'files/1/serverpath'
+	struct FCbEntryDifference
+	{
+		FUtf8String PropertyName;
+		FCbFieldView OldValue;
+		FCbFieldView NewValue;
+	};
+
+	// Recursively diff the contents of a compact binary field/object
+	// Important! Arrays are assumed to be sorted!
+	UE_API TArray<FCbEntryDifference> DiffCompactBinary(const FCbFieldView& cb1, const FCbFieldView& cb2);
+
+}	// namespace UE
+
+#undef UE_API
diff --git a/Engine/Source/Developer/ZenOplogUtils/Public/Experimental/ZenOplogDiff.h b/Engine/Source/Developer/ZenOplogUtils/Public/Experimental/ZenOplogDiff.h
new file mode 100644
--- /dev/null
+++ b/Engine/Source/Developer/ZenOplogUtils/Public/Experimental/ZenOplogDiff.h
@@ -0,0 +1,65 @@
+// Copyright Epic Games, Inc. All Rights Reserved.
+
+#pragma once
+
+#include "CoreTypes.h"
+#include "Templates/Tuple.h"
+#include "Serialization/CompactBinary.h"
+#include "Serialization/JsonWriter.h"
+#include "Misc/EnumClassFlags.h"
+#include "Experimental/DiffCompactBinary.h"
+
+#define UE_API ZENOPLOGUTILS_API
+
+namespace UE
+{
+	struct FOplogManifest;
+	struct FCbEntryDifference;
+
+	// Represents a changed Op and the properties with different values
+	using FOplogDiffChangedOp = TPair<FUtf8String, TArray<FCbEntryDifference>>;
+
+	// Represents the full diff of an oplog/manifest
+	// Takes shared ownership of the compact binary in a manifest
+	struct FOplogDiffResults
+	{
+		TArray<FUtf8String> OpsMissingInManifest1;
+		TArray<FUtf8String> OpsMissingInManifest2;
+		TArray<FUtf8String> IdenticalOps;
+		TArray<FOplogDiffChangedOp> ChangedOpsWithSameOutput;		// Op has differences, but the output data is the same
+		TArray<FOplogDiffChangedOp> ChangedOpsWithOutputDifferences;	// Op has cooked output differences
+
+		FCbField Manifest1Data;
+		FCbField Manifest2Data;
+	};
+
+	// Diff all ops, reports all changes in each op. Results are unsorted
+	UE_API FOplogDiffResults DiffManifests(const FOplogManifest& Manifest1, const FOplogManifest& Manifest2);
+
+	// Sorts the results of the diff
+	UE_API void SortDiffResults(FOplogDiffResults& FullResults);
+
+	// Returns true if any diff entry describes an array modification (array sizes different),
+	// or if an entry in that array contains a change of a specific property
+	// i.e. if changed property name == 'ArrayName' or 'ArrayName/.../PropertyName'
+	UE_API bool OpDiffContainsArrayValueChange(const FOplogDiffChangedOp& Op, FUtf8StringView ArrayName, FUtf8StringView PropertyName);
+
+	// Output diff results
+	
+	enum class EOutputManifestDiffOptions : int32
+	{
+		None = 0,
+		OutputDifferences = 1,		// Output the values of any different properties
+		OutputIdenticalOps = 2		// Output the list of identical ops
+	};
+	ENUM_CLASS_FLAGS(EOutputManifestDiffOptions);
+
+	// Flags = a combination of EOutputManifestDiffOptions
+	UE_API void OutputManifestDiffResultsToJson(TSharedRef<TJsonWriter<>>& Writer, const FOplogDiffResults& Results, EOutputManifestDiffOptions Flags );
+
+	// Flags = a combination of EOutputManifestDiffOptions
+	UE_API void OutputManifestDiffResultsToCompactBinary(FCbWriter& Writer, const FOplogDiffResults& Results, EOutputManifestDiffOptions Flags);
+
+}	// namespace UE
+
+#undef UE_API
diff --git a/Engine/Source/Developer/ZenOplogUtils/Public/Experimental/ZenOplogManifest.h b/Engine/Source/Developer/ZenOplogUtils/Public/Experimental/ZenOplogManifest.h
new file mode 100644
--- /dev/null
+++ b/Engine/Source/Developer/ZenOplogUtils/Public/Experimental/ZenOplogManifest.h
@@ -0,0 +1,107 @@
+// Copyright Epic Games, Inc. All Rights Reserved.
+
+#pragma once
+
+#include "Containers/Map.h"
+#include "Containers/Utf8String.h"
+#include "IO/IoHash.h"
+#include "Serialization/CompactBinary.h"
+#include "Templates/Tuple.h"
+
+#define UE_API ZENOPLOGUTILS_API
+
+namespace UE
+{
+
+	// A manifest is a chunk of compact binary representing the cooked output of a build.
+	struct FOplogManifest
+	{
+		// Caches a view of an op in the compact binary and provides accessors to known values
+		struct FOp
+		{
+			UE_API FUtf8String GetKey() const;
+			UE_API FIoHash GetArtifactHash() const;
+			UE_API FIoHash CookImportExport() const;
+			UE_API FIoHash GetLogs() const;
+
+			struct FPackageStoreEntry
+			{
+				uint32 Flags = 0;
+				FUtf8String PackageName;
+				TArray<uint64> ImportedPackages;
+				TArray<uint64> SoftPackageReferences;
+			};
+			UE_API FPackageStoreEntry GetPackageStoreEntry() const;
+
+			struct FPackageData
+			{
+				FCbObjectId ID;
+				int64 Size = 0;
+				int64 RawSize = 0;
+				FIoHash DataHash;
+				FUtf8String Filename;
+			};
+			UE_API TArray<FPackageData> GetPackageDatas() const;
+
+			struct FBulkData
+			{
+				FCbObjectId ID;
+				FUtf8String TypeStr;
+				int64 Size = 0;
+				int64 RawSize = 0;
+				FIoHash DataHash;
+				FUtf8String Filename;
+			};
+			UE_API TArray<FBulkData> GetBulkDatas() const;
+
+			struct FFile
+			{
+				FCbObjectId ID;
+				FUtf8String ServerPath;
+				FUtf8String ClientPath;
+				FIoHash DataHash;
+			};
+			UE_API TArray<FFile> GetFiles() const;
+
+			struct FMeta
+			{
+				FCbObjectId ID;
+				FUtf8String Name;
+				FIoHash DataHash;
+			};
+			UE_API TArray<FMeta> GetMetas() const;
+
+			FCbFieldView CbData;	// Cache a field view to the op data in the manifest compact binary
+		};
+
+		// Caches op keys + hashes, and populates OpKeyToIndex. Call after loading. Returns false on unexpected data
+		UE_API bool RebuildCachedData();
+
+		FCbField AllOpsCbData;
+		TArray<FOp> Ops;
+		TArray<TPair<FUtf8String, uint32>> OpKeysAndHashes;
+		TMap<FUtf8String, int32> OpKeyToIndex;	// Op key -> index into Ops array
+	};
+
+	struct FLoadOplogManifestResult
+	{
+		enum class EStatus
+		{
+			Ok,
+			ErrorFileNotFound,
+			ErrorUnsupportedFormat,
+			ErrorMalformedCompactBinary,
+			Error
+		};
+		FLoadOplogManifestResult(EStatus Status, TOptional<FOplogManifest>&& Manifest = {});
+
+		EStatus Result;
+		TOptional<FOplogManifest> Manifest;
+	};
+
+	UE_API FLoadOplogManifestResult LoadOplogManifestFromCompactBinary(const FCbObject& CompactBinary);
+	UE_API FLoadOplogManifestResult LoadOplogManifestFromFile(FString Path);
+
+}	// namespace UE
+
+#undef UE_API
diff --git a/Engine/Source/Developer/ZenOplogUtils/Public/Experimental/ZenOplogSnapshots.h b/Engine/Source/Developer/ZenOplogUtils/Public/Experimental/ZenOplogSnapshots.h
new file mode 100644
--- /dev/null
+++ b/Engine/Source/Developer/ZenOplogUtils/Public/Experimental/ZenOplogSnapshots.h
@@ -0,0 +1,50 @@
+// Copyright Epic Games, Inc. All Rights Reserved.
+
+#pragma once
+
+#include "CoreTypes.h"
+#include "Async/Future.h"
+#include "Containers/UnrealString.h"
+
+#define UE_API ZENOPLOGUTILS_API
+
+namespace UE
+{
+	// Contains the data required to get a oplog/snapshot from the cloud servers
+	// Can be parsed from exported zen snapshot json, or found from the build service
+	struct FZenSnapshotDescriptor
+	{
+		FString CloudHost;
+		FString Namespace;
+		FString Bucket;
+		FString BuildID;
+	};
+
+	// Attempt to parse an exported zen snapshot json from horde
+	UE_API TArray<FZenSnapshotDescriptor> ParseSnapshotDescriptorFromJson(FStringView PathToJson);
+
+	// Try to find a snapshot descriptor matching a build from the zen build service
+	// Returns an empty TOptional if the build was not found, or any errors occur
+	// Note this relies on callbacks being fired on the game thread - do not do a blocking wait/get!
+	//	bLogAllBuilds - logs all namespaces, buckets, and CLs found
+	UE_API TFuture<TOptional<FZenSnapshotDescriptor>> FindSnapshotDescriptorForBuild(FStringView ProjectName, FStringView StreamName, FStringView PlatformName, FStringView CommitID, bool bLogAllBuilds = false);
+
+
+	// Snapshot downloads.
+
+	struct FDownloadOplogManifestResult
+	{
+		enum class EStatus
+		{
+			Ok,
+			Error
+		};
+		EStatus Result;
+		FString DownloadedFilename;		// Path to the local manifest file after a successful download
+	};
+
+	// Download an oplog manifest to a local file.
+	UE_API FDownloadOplogManifestResult DownloadOplogManifest(const FZenSnapshotDescriptor& Descriptor, FString OutputDirectory = {}, bool bDownloadAsJson = false);
+}
+
+#undef UE_API
\ No newline at end of file
diff --git a/Engine/Source/Developer/ZenOplogUtils/Tests/ZenOplogUtilsTest.cpp b/Engine/Source/Developer/ZenOplogUtils/Tests/ZenOplogUtilsTest.cpp
new file mode 100644
--- /dev/null
+++ b/Engine/Source/Developer/ZenOplogUtils/Tests/ZenOplogUtilsTest.cpp
@@ -0,0 +1,1036 @@
+// Copyright Epic Games, Inc. All Rights Reserved.
+
+#if WITH_LOW_LEVEL_TESTS
+
+#include "TestHarness.h"
+#include "Experimental/DiffCompactBinary.h"
+#include "Experimental/ZenOplogDiff.h"
+#include "Experimental/ZenOplogManifest.h"
+#include "Algo/IsSorted.h"
+#include "Algo/Transform.h"
+#include "Serialization/CompactBinarySerialization.h"
+
+#if WITH_EDITORONLY_DATA
+
+namespace UE
+{
+	namespace TestData	// Helpers for building op compact binary in the formats we expect
+	{
+		void WriteOp(FCbWriter& Writer, FString Key, FIoHash ArtifactHash = {}, FIoHash ImportExportHash = {}, FIoHash LogsHash = {})
+		{
+			Writer.AddString("key", Key);
+			if (!ArtifactHash.IsZero())
+			{
+				Writer.AddBinaryAttachment("meta.cook.artifacts", ArtifactHash);
+			}
+			if (!ImportExportHash.IsZero())
+			{
+				Writer.AddBinaryAttachment("meta.cook.importexport", ImportExportHash);
+			}
+			if (!LogsHash.IsZero())
+			{
+				Writer.AddBinaryAttachment("meta.cook.logs", LogsHash);
+			}
+		}
+
+		void WritePackageStoreEntry(FCbWriter& Writer, FString PackageName, int32 Flags, TArray<int64> ImportedPackages = {}, TArray<int64> SoftReferences = {})
+		{
+			Writer.BeginObject("packagestoreentry");
+			Writer.AddInteger("flags", Flags);
+			Writer.AddString("packagename", PackageName);
+			if (ImportedPackages.Num() > 0)
+			{
+				Writer.BeginArray("importedpackageids");
+				for (uint64 ID : ImportedPackages)
+				{
+					Writer.AddInteger(ID);
+				}
+				Writer.EndArray();	// importedpackageids
+			}
+			if (SoftReferences.Num() > 0)
+			{
+				Writer.BeginArray("softpackagereferences");
+				for (uint64 ID : SoftReferences)
+				{
+					Writer.AddInteger(ID);
+				}
+				Writer.EndArray();	// softpackagereferences
+			}
+			Writer.EndObject();	// packagestoreentry
+		}
+
+		void WritePackageData(FCbWriter& Writer, const FOplogManifest::FOp::FPackageData& Data)
+		{
+			Writer.BeginObject();
+			Writer.AddObjectId("id", Data.ID);
+			Writer.AddInteger("size", Data.Size);
+			Writer.AddInteger("rawsize", Data.RawSize);
+			Writer.AddBinaryAttachment("data", Data.DataHash);
+			Writer.AddString("filename", Data.Filename);
+			Writer.EndObject();
+		}
+
+		void WriteBulkData(FCbWriter& Writer, const FOplogManifest::FOp::FBulkData& Data)
+		{
+			Writer.BeginObject();
+			Writer.AddObjectId("id", Data.ID);
+			Writer.AddString("type", Data.TypeStr);
+			Writer.AddInteger("size", Data.Size);
+			Writer.AddInteger("rawsize", Data.RawSize);
+			Writer.AddBinaryAttachment("data", Data.DataHash);
+			Writer.AddString("filename", Data.Filename);
+			Writer.EndObject();
+		}
+
+		void WriteFile(FCbWriter& Writer, const FOplogManifest::FOp::FFile& Data)
+		{
+			Writer.BeginObject();
+			Writer.AddObjectId("id", Data.ID);
+			Writer.AddString("serverpath", Data.ServerPath);
+			Writer.AddString("clientpath", Data.ClientPath);
+			Writer.AddBinaryAttachment("data", Data.DataHash);
+			Writer.EndObject();
+		}
+
+		void WriteMeta(FCbWriter& Writer, const FOplogManifest::FOp::FMeta& Data)
+		{
+			Writer.BeginObject();
+			Writer.AddObjectId("id", Data.ID);
+			Writer.AddString("name", Data.Name);
+			Writer.AddBinaryAttachment("data", Data.DataHash);
+			Writer.EndObject();
+		}
+
+		FOplogManifest GenerateSimpleTestManifest(FString Key, FString PackageStoreName, FIoHash PackageHash, int64 PackageSize, FUtf8String PackageDataFilename, TArray<int64> ImportExports = {})
+		{
+			FCbWriter CbWriter;
+			CbWriter.BeginObject();
+			CbWriter.BeginArray("ops");
+			CbWriter.BeginObject();
+			WriteOp(CbWriter, Key);
+			WritePackageStoreEntry(CbWriter, PackageStoreName, 13, ImportExports, {});
+			CbWriter.BeginArray("packagedata");
+			FOplogManifest::FOp::FPackageData PackageData{
+				.ID = FCbObjectId(FCbObjectId::ByteArray{10}),
+				.Size = PackageSize,
+				.RawSize = 32,
+				.DataHash = PackageHash,
+				.Filename = PackageDataFilename
+			};
+			WritePackageData(CbWriter, PackageData);
+			CbWriter.EndArray();	//packagedata
+			CbWriter.EndObject();
+			CbWriter.EndArray();	// ops
+			CbWriter.EndObject();
+			FLoadOplogManifestResult LoadResult = LoadOplogManifestFromCompactBinary(CbWriter.Save().AsObject());
+			CHECK(LoadResult.Result == FLoadOplogManifestResult::EStatus::Ok);
+			return *LoadResult.Manifest;
+		}
+
+		FOplogManifest GenerateCustomTestManifest(TFunction<void(FCbWriter&)> WriteOpsArrayFn)
+		{
+			FCbWriter CbWriter;
+			CbWriter.BeginObject();
+			CbWriter.BeginArray("ops");
+			WriteOpsArrayFn(CbWriter);
+			CbWriter.EndArray();	// ops
+			CbWriter.EndObject();
+			FLoadOplogManifestResult LoadResult = LoadOplogManifestFromCompactBinary(CbWriter.Save().AsObject());
+			CHECK(LoadResult.Result == FLoadOplogManifestResult::EStatus::Ok);
+			return *LoadResult.Manifest;
+		}
+
+		FOplogManifest GenerateManifestWithMultipleOps(TArray<FString> OpKeys)
+		{
+			FCbWriter CbWriter;
+			CbWriter.BeginObject();
+			CbWriter.BeginArray("ops");
+			for (const FString& Key : OpKeys)
+			{
+				CbWriter.BeginObject();
+				WriteOp(CbWriter, Key);
+				CbWriter.EndObject();
+			}
+			CbWriter.EndArray();	// ops
+			CbWriter.EndObject();
+			FLoadOplogManifestResult LoadResult = LoadOplogManifestFromCompactBinary(CbWriter.Save().AsObject());
+			CHECK(LoadResult.Result == FLoadOplogManifestResult::EStatus::Ok);
+			return *LoadResult.Manifest;
+		}
+	}
+
+	TEST_CASE("ZenOplogUtils::Manifest::EmptyArchive", "[ZenOplogUtils]")
+	{
+		FCbObject EmptyObject;
+		FLoadOplogManifestResult Result = LoadOplogManifestFromCompactBinary(EmptyObject);
+		CHECK(Result.Result == FLoadOplogManifestResult::EStatus::ErrorMalformedCompactBinary);
+	}
+
+	TEST_CASE("ZenOplogUtils::Manifest::UnexpectedDataNoOps", "[ZenOplogUtils]")
+	{
+		FCbWriter CbWriter;
+		CbWriter.BeginObject();
+		CbWriter.BeginArray("Blah");
+		CbWriter.EndArray();
+		CbWriter.EndObject();
+		FCbObject Cb = CbWriter.Save().AsObject();
+		FLoadOplogManifestResult Result = LoadOplogManifestFromCompactBinary(Cb);
+		CHECK(Result.Result == FLoadOplogManifestResult::EStatus::ErrorMalformedCompactBinary);
+	}
+
+	TEST_CASE("ZenOplogUtils::Manifest::EmptyOpsList", "[ZenOplogUtils]")
+	{
+		FCbWriter CbWriter;
+		CbWriter.BeginObject();
+		CbWriter.BeginArray("ops");
+		CbWriter.EndArray();
+		CbWriter.EndObject();
+		FCbObject Cb = CbWriter.Save().AsObject();
+		FLoadOplogManifestResult Result = LoadOplogManifestFromCompactBinary(Cb);
+		CHECK(Result.Result == FLoadOplogManifestResult::EStatus::Ok);
+	}
+
+	TEST_CASE("ZenOplogUtils::Manifest::OpWithNoKey", "[ZenOplogUtils]")
+	{
+		FCbWriter CbWriter;
+		CbWriter.BeginObject();
+		CbWriter.BeginArray("ops");
+		CbWriter.BeginObject();
+		CbWriter.BeginObject("packagestoreentry");
+		CbWriter.AddInteger("flags", 1);
+		CbWriter.EndObject();
+		CbWriter.EndObject();
+		CbWriter.EndArray();
+		CbWriter.EndObject();
+		FCbObject Cb = CbWriter.Save().AsObject();
+		FLoadOplogManifestResult Result = LoadOplogManifestFromCompactBinary(Cb);
+		CHECK(Result.Result == FLoadOplogManifestResult::EStatus::ErrorMalformedCompactBinary);
+	}
+
+	TEST_CASE("ZenOplogUtils::Manifest::OpsWithDuplicateKey", "[ZenOplogUtils]")
+	{
+		FCbWriter CbWriter;
+		CbWriter.BeginObject();
+		CbWriter.BeginArray("ops");
+		CbWriter.BeginObject();
+		TestData::WriteOp(CbWriter, "ThisIsADuplicate");
+		CbWriter.EndObject();
+		CbWriter.BeginObject();
+		TestData::WriteOp(CbWriter, "ThisIsOK");
+		CbWriter.EndObject();
+		CbWriter.BeginObject();
+		TestData::WriteOp(CbWriter, "ThisIsADuplicate");
+		CbWriter.EndObject();
+		CbWriter.EndArray();
+		CbWriter.EndObject();
+		FCbObject Cb = CbWriter.Save().AsObject();
+		FLoadOplogManifestResult Result = LoadOplogManifestFromCompactBinary(Cb);
+		CHECK(Result.Result == FLoadOplogManifestResult::EStatus::Error);
+	}
+
+	TEST_CASE("ZenOplogUtils::Manifest::OpsWithUnexpectedData", "[ZenOplogUtils]")
+	{
+		FCbWriter CbWriter;
+		CbWriter.BeginObject();
+		CbWriter.BeginArray("ops");
+		CbWriter.BeginObject();
+		TestData::WriteOp(CbWriter, "Op1");
+		CbWriter.AddString("somenewstring", "Blah");
+		CbWriter.AddInteger("somenewinteger", 32);
+		CbWriter.EndObject();
+		CbWriter.BeginObject();
+		TestData::WriteOp(CbWriter, "Op2");
+		CbWriter.AddString("someotherstring", "Hmmm");
+		CbWriter.BeginArray("SomeNewArray");
+		CbWriter.BeginObject();
+		CbWriter.AddBinaryAttachment("SomeNewAttachment", FIoHash(TEXT("aaaaaaaaabbbbbbbbbbddddddddddddddddddddd")));
+		CbWriter.AddString("FilenameOrSomething", "Test");
+		CbWriter.EndObject();
+		CbWriter.EndArray();
+		CbWriter.EndObject();
+		CbWriter.EndArray();	// ops
+		CbWriter.EndObject();
+		FLoadOplogManifestResult Result = LoadOplogManifestFromCompactBinary(CbWriter.Save().AsObject());
+		CHECK(Result.Result == FLoadOplogManifestResult::EStatus::Ok);
+	}
+
+	TEST_CASE("ZenOplogUtils::Manifest::OpAccessors", "[ZenOplogUtils]")
+	{
+		const FOplogManifest::FOp::FPackageData TestPackage
+		{
+			.ID = FCbObjectId(FCbObjectId::ByteArray{0}),
+			.Size = 25,
+			.RawSize = 64,
+			.DataHash = FIoHash(TEXT("ed55912a82d1c88c4b0e73b1a0740277a82d732d")),
+			.Filename = "TestFile1.uasset"
+		};
+
+		const FOplogManifest::FOp::FBulkData TestBulkData
+		{
+			.ID = FCbObjectId(FCbObjectId::ByteArray{1}),
+			.TypeStr = "Mmap",
+			.Size = 1024,
+			.RawSize = 1400,
+			.DataHash = FIoHash(TEXT("aaaabbbbccccdddd4b0e73b1a0740277a82d732d")),
+			.Filename = "SomeFileWithLotsOfData.ubulk"
+		};
+
+		const FOplogManifest::FOp::FFile TestFile
+		{
+			.ID = FCbObjectId(FCbObjectId::ByteArray{2}),
+			.ServerPath = "Path//On//Server",
+			.ClientPath = "Path//On//Client",
+			.DataHash = FIoHash(TEXT("aaaabbbbccccddddeeeeeeffffff0277a82d732d"))
+		};
+
+		const FOplogManifest::FOp::FMeta TestMeta
+		{
+			.ID = FCbObjectId(FCbObjectId::ByteArray{3}),
+			.Name = "SomeName",
+			.DataHash = FIoHash(TEXT("aaaabbabccccdddd2e6e2ef6ff2f2267a82d732d"))
+		};
+
+		const FIoHash ArtifactsHash(TEXT("d3d7a8a4c3a26b7b3a8234588213287a82d732da"));
+		const FIoHash ImportExportHash(TEXT("11112225777322bbaa2627374237373782d7314b"));
+		const FIoHash LogsHash(TEXT("cccccccccccbbbbbbbbbbbbbbbbbbbbbaaaaadde"));
+
+		FCbWriter CbWriter;
+		CbWriter.BeginObject();
+		CbWriter.BeginArray("ops");
+		CbWriter.BeginObject();
+		TestData::WriteOp(CbWriter, "TestOp1", ArtifactsHash, ImportExportHash, LogsHash);
+		TestData::WritePackageStoreEntry(CbWriter, "TestOp1Package", 42, { 1, 2, 3 }, { 8, 7, 6 });
+		CbWriter.BeginArray("packagedata");
+		TestData::WritePackageData(CbWriter, TestPackage);
+		CbWriter.EndArray();	//packagedata
+		CbWriter.BeginArray("bulkdata");
+		TestData::WriteBulkData(CbWriter, TestBulkData);
+		CbWriter.EndArray();	// bulkdata
+		CbWriter.BeginArray("files");
+		TestData::WriteFile(CbWriter, TestFile);
+		CbWriter.EndArray();	// files
+		CbWriter.BeginArray("meta");
+		TestData::WriteMeta(CbWriter, TestMeta);
+		CbWriter.EndArray();	// meta
+		CbWriter.EndObject();
+		CbWriter.EndArray();	// ops
+		CbWriter.EndObject();
+		FCbObject Cb = CbWriter.Save().AsObject();
+
+		FLoadOplogManifestResult Result = LoadOplogManifestFromCompactBinary(Cb);
+		CHECK(Result.Result == FLoadOplogManifestResult::EStatus::Ok);
+
+		CHECK(Result.Manifest->Ops.Num() == 1);
+		const FOplogManifest::FOp& TestOp = Result.Manifest->Ops[0];
+		CHECK(TestOp.GetKey() == "TestOp1");
+		CHECK(TestOp.GetArtifactHash() == ArtifactsHash);
+		CHECK(TestOp.CookImportExport() == ImportExportHash);
+		CHECK(TestOp.GetLogs() == LogsHash);
+
+		FOplogManifest::FOp::FPackageStoreEntry PStoreEntry = TestOp.GetPackageStoreEntry();
+		CHECK(PStoreEntry.Flags == 42);
+		CHECK(PStoreEntry.PackageName == "TestOp1Package");
+		CHECK(PStoreEntry.ImportedPackages == TArray<uint64>{1, 2, 3});
+		CHECK(PStoreEntry.SoftPackageReferences == TArray<uint64>{8, 7, 6});
+
+		TArray<FOplogManifest::FOp::FPackageData> Packages = TestOp.GetPackageDatas();
+		CHECK(Packages.Num() == 1);
+		CHECK(Packages[0].ID == TestPackage.ID);
+		CHECK(Packages[0].Size == TestPackage.Size);
+		CHECK(Packages[0].RawSize == TestPackage.RawSize);
+		CHECK(Packages[0].DataHash == TestPackage.DataHash);
+		CHECK(Packages[0].Filename == TestPackage.Filename);
+
+		TArray<FOplogManifest::FOp::FBulkData> BulkDatas = TestOp.GetBulkDatas();
+		CHECK(BulkDatas.Num() == 1);
+		CHECK(BulkDatas[0].ID == TestBulkData.ID);
+		CHECK(BulkDatas[0].TypeStr == TestBulkData.TypeStr);
+		CHECK(BulkDatas[0].Size == TestBulkData.Size);
+		CHECK(BulkDatas[0].RawSize == TestBulkData.RawSize);
+		CHECK(BulkDatas[0].DataHash == TestBulkData.DataHash);
+		CHECK(BulkDatas[0].Filename == TestBulkData.Filename);
+
+		TArray<FOplogManifest::FOp::FFile> Files = TestOp.GetFiles();
+		CHECK(Files.Num() == 1);
+		CHECK(Files[0].ID == TestFile.ID);
+		CHECK(Files[0].ServerPath == TestFile.ServerPath);
+		CHECK(Files[0].ClientPath == TestFile.ClientPath);
+		CHECK(Files[0].DataHash == TestFile.DataHash);
+
+		TArray<FOplogManifest::FOp::FMeta> Metas = TestOp.GetMetas();
+		CHECK(Metas.Num() == 1);
+		CHECK(Metas[0].ID == TestMeta.ID);
+		CHECK(Metas[0].Name == TestMeta.Name);
+		CHECK(Metas[0].DataHash == TestMeta.DataHash);
+	}
+
+const FOplogDiffChangedOp* FindChangedOp(const FOplogDiffResults& Diff, FUtf8String OpKeyToFind)
+{
+	const FOplogDiffChangedOp* Changed = Diff.ChangedOpsWithSameOutput.FindByPredicate([&OpKeyToFind](const FOplogDiffChangedOp& Entry)
+	{
+		return Entry.Key == OpKeyToFind;
+	});
+	if (Changed != nullptr)
+	{
+		return Changed;
+	}
+	else
+	{
+		return Diff.ChangedOpsWithOutputDifferences.FindByPredicate([&OpKeyToFind](const FOplogDiffChangedOp& Entry)
+		{
+			return Entry.Key == OpKeyToFind;
+		});
+	}
+};
+
+bool ChangedOpPropertyExists(const FOplogDiffResults& Diff, FUtf8String OpKeyToFind, FUtf8String PropertyName)
+{
+	const FOplogDiffChangedOp* Changes = FindChangedOp(Diff, OpKeyToFind);
+	if (Changes)
+	{
+		for (const FCbEntryDifference& Difference : Changes->Value)
+		{
+			if (Difference.PropertyName == PropertyName)
+			{
+				return true;
+			}
+		}
+	}
+	return false;
+};
+
+// Helper to do diff(manifest1,manifest2) and diff(manifest2,manifest1), with the same tests
+void DiffBothDirections(const FOplogManifest& Manifest1, const FOplogManifest& Manifest2, TFunction<void(const FOplogDiffResults&)> CheckFn)
+{
+	FOplogDiffResults Diff1 = DiffManifests(Manifest1, Manifest2);
+	FOplogDiffResults Diff2 = DiffManifests(Manifest2, Manifest1);
+	CheckFn(Diff1);
+	CheckFn(Diff2);
+};
+
+TEST_CASE("ZenOplogUtils::Diff::SameManifestIsEqual", "[ZenOplogUtils]")
+{
+	FOplogManifest SimpleManifest1 = TestData::GenerateSimpleTestManifest("TestOp1", 
+		"TestPackage1", 
+		FIoHash(TEXT("aaaaaabbbbbbbcccccccddddddeeeeeeffffffff")), 
+		32, 
+		"TestPackageData1", 
+		{ 1, 2, 3 });
+	DiffBothDirections(SimpleManifest1, SimpleManifest1, [](const FOplogDiffResults& Diff) 
+	{
+		CHECK(Diff.IdenticalOps.Num() == 1);
+		CHECK(Diff.OpsMissingInManifest1.Num() == 0);
+		CHECK(Diff.OpsMissingInManifest2.Num() == 0);
+		CHECK(Diff.ChangedOpsWithSameOutput.Num() == 0);
+		CHECK(Diff.ChangedOpsWithOutputDifferences.Num() == 0);
+	});
+}
+
+TEST_CASE("ZenOplogUtils::Diff::MultipleSimilarOps", "[ZenOplogUtils]")
+{
+	// Same ops, different order
+	FOplogManifest Manifest1 = TestData::GenerateManifestWithMultipleOps(
+	{
+		"SomeOp",
+		"SomeOtherOp",
+		"YetAnotherOp",
+		"ReallyTheresMoreOps"
+	});
+	FOplogManifest Manifest2 = TestData::GenerateManifestWithMultipleOps(
+	{
+		"YetAnotherOp",
+		"SomeOtherOp",
+		"ReallyTheresMoreOps",
+		"SomeOp"
+	});
+	DiffBothDirections(Manifest1, Manifest2, [](const FOplogDiffResults& Diff) 
+	{
+		CHECK(Diff.IdenticalOps.Num() == 4);
+		CHECK(Diff.OpsMissingInManifest1.Num() == 0);
+		CHECK(Diff.OpsMissingInManifest2.Num() == 0);
+		CHECK(Diff.ChangedOpsWithSameOutput.Num() == 0);
+		CHECK(Diff.ChangedOpsWithOutputDifferences.Num() == 0);
+	});
+}
+
+TEST_CASE("ZenOplogUtils::Diff::DifferentOps", "[ZenOplogUtils]")
+{
+	FOplogManifest Manifest1 = TestData::GenerateManifestWithMultipleOps(
+	{
+		"Op1",
+		"Op2",
+		"Op3",
+		"Op4",
+		"Op5"
+	});
+	FOplogManifest Manifest2 = TestData::GenerateManifestWithMultipleOps(
+	{
+		"Op1",
+		"OpA",
+		"Op3",
+		"OpB",
+		"Op5"
+	});
+	FOplogDiffResults Diff = DiffManifests(Manifest1, Manifest2);
+	CHECK(Diff.IdenticalOps.Num() == 3);
+	CHECK(Diff.OpsMissingInManifest1.Num() == 2);
+	CHECK(Diff.OpsMissingInManifest1.Contains("OpA"));
+	CHECK(Diff.OpsMissingInManifest1.Contains("OpB"));
+	CHECK(Diff.OpsMissingInManifest2.Num() == 2);
+	CHECK(Diff.OpsMissingInManifest2.Contains("Op2"));
+	CHECK(Diff.OpsMissingInManifest2.Contains("Op4"));
+	CHECK(Diff.ChangedOpsWithSameOutput.Num() == 0);
+	CHECK(Diff.ChangedOpsWithOutputDifferences.Num() == 0);
+}
+
+TEST_CASE("ZenOplogUtils::Diff::DifferentPackageName", "[ZenOplogUtils]")
+{
+	FOplogManifest SimpleManifest1 = TestData::GenerateSimpleTestManifest("TestOp1", 
+		"TestPackage1", 
+		FIoHash(TEXT("aaaaaabbbbbbbcccccccddddddeeeeeeffffffff")),
+		32, 
+		"TestPackageData1", 
+		{ 1, 2, 3 });
+	FOplogManifest SimpleManifestDifferentPackageName = TestData::GenerateSimpleTestManifest("TestOp1", 
+		"SomePackageName", 
+		FIoHash(TEXT("aaaaaabbbbbbbcccccccddddddeeeeeeffffffff")), 
+		32, 
+		"TestPackageData1", 
+		{ 1, 2, 3 });
+	DiffBothDirections(SimpleManifest1, SimpleManifestDifferentPackageName, [](const FOplogDiffResults& Diff) 
+	{
+		CHECK(Diff.IdenticalOps.Num() == 0);
+		CHECK(Diff.OpsMissingInManifest1.Num() == 0);
+		CHECK(Diff.OpsMissingInManifest2.Num() == 0);
+		CHECK(Diff.ChangedOpsWithSameOutput.Num() == 1);
+		CHECK(Diff.ChangedOpsWithOutputDifferences.Num() == 0);
+		CHECK(FindChangedOp(Diff, "TestOp1"));
+		CHECK(ChangedOpPropertyExists(Diff, "TestOp1", "packagestoreentry/packagename"));
+	});
+}
+
+TEST_CASE("ZenOplogUtils::Diff::DifferentImportExportArraySize", "[ZenOplogUtils]")
+{
+	FOplogManifest SimpleManifest1 = TestData::GenerateSimpleTestManifest("TestOp1", 
+		"TestPackage1", 
+		FIoHash(TEXT("aaaaaabbbbbbbcccccccddddddeeeeeeffffffff")), 
+		32, 
+		"TestPackageData1", 
+		{ 1, 2, 3 });
+	FOplogManifest SimpleManifestDifferentImportExportSize = TestData::GenerateSimpleTestManifest("TestOp1", 
+		"TestPackage1", 
+		FIoHash(TEXT("aaaaaabbbbbbbcccccccddddddeeeeeeffffffff")), 
+		32, 
+		"TestPackageData1", 
+		{ 6 });
+	DiffBothDirections(SimpleManifest1, SimpleManifestDifferentImportExportSize, [](const FOplogDiffResults& Diff) 
+	{
+		CHECK(Diff.IdenticalOps.Num() == 0);
+		CHECK(Diff.OpsMissingInManifest1.Num() == 0);
+		CHECK(Diff.OpsMissingInManifest2.Num() == 0);
+		CHECK(Diff.ChangedOpsWithSameOutput.Num() == 1);
+		CHECK(Diff.ChangedOpsWithOutputDifferences.Num() == 0);
+		CHECK(FindChangedOp(Diff, "TestOp1"));
+		CHECK(ChangedOpPropertyExists(Diff, "TestOp1", "packagestoreentry/importedpackageids"));
+	});
+}
+
+TEST_CASE("ZenOplogUtils::Diff::ManyDifferences", "[ZenOplogUtils]")
+{
+	FOplogManifest SimpleManifest1 = TestData::GenerateSimpleTestManifest("TestOp1", 
+		"TestPackage1", 
+		FIoHash(TEXT("aaaaaabbbbbbbcccccccddddddeeeeeeffffffff")), 
+		32, 
+		"TestPackageData1", 
+		{ 1, 2, 3 });
+	FOplogManifest SimpleManifestAllDifference = TestData::GenerateSimpleTestManifest("TestOp1", 
+		"Huh", 
+		FIoHash(TEXT("1111222233334444ccccddddddeeeeeeffffffff")), 
+		111, 
+		"Whatt", 
+		{ 4, 2, 5 });
+	DiffBothDirections(SimpleManifest1, SimpleManifestAllDifference, [](const FOplogDiffResults& Diff) 
+	{
+		CHECK(Diff.IdenticalOps.Num() == 0);
+		CHECK(Diff.OpsMissingInManifest1.Num() == 0);
+		CHECK(Diff.OpsMissingInManifest2.Num() == 0);
+		CHECK(Diff.ChangedOpsWithSameOutput.Num() == 0);
+		CHECK(Diff.ChangedOpsWithOutputDifferences.Num() == 1);
+		CHECK(FindChangedOp(Diff, "TestOp1"));
+		CHECK(Diff.ChangedOpsWithOutputDifferences[0].Value.Num() == 6);
+		CHECK(ChangedOpPropertyExists(Diff, "TestOp1", "packagestoreentry/packagename"));
+		CHECK(ChangedOpPropertyExists(Diff, "TestOp1", "packagestoreentry/importedpackageids/0"));
+		CHECK(ChangedOpPropertyExists(Diff, "TestOp1", "packagestoreentry/importedpackageids/2"));
+		CHECK(ChangedOpPropertyExists(Diff, "TestOp1", "packagedata/0/filename"));
+		CHECK(ChangedOpPropertyExists(Diff, "TestOp1", "packagedata/0/size"));
+		CHECK(ChangedOpPropertyExists(Diff, "TestOp1", "packagedata/0/data"));
+	});
+}
+
+TEST_CASE("ZenOplogUtils::Diff::DataMissingInOps", "[ZenOplogUtils]")
+{
+	auto OpsWithArtifactsOrImportExports = [](FCbWriter& Writer)
+	{
+		Writer.BeginObject();
+		TestData::WriteOp(Writer, "Op1", FIoHash(TEXT("aaaabbbbccccdddd4b0e73b1a0740277a82d732d")));
+		Writer.EndObject();
+		Writer.BeginObject();
+		TestData::WriteOp(Writer, "Op2", {}, FIoHash(TEXT("eeeeeeeffffffffaaaaaabbbbbbaaaaaaaa2222d")));
+		Writer.EndObject();
+		Writer.BeginObject();
+		TestData::WriteOp(Writer, "Op3", FIoHash(TEXT("6666666677777777778888888888844444411111")), FIoHash(TEXT("1111111112222222222233333333444444444444")));
+		Writer.EndObject();
+	};
+	auto OpsWithMissingArtifactsOrImportExports = [](FCbWriter& Writer)
+	{
+		Writer.BeginObject();
+		TestData::WriteOp(Writer, "Op1", {});
+		Writer.EndObject();
+		Writer.BeginObject();
+		TestData::WriteOp(Writer, "Op2", {}, {});
+		Writer.EndObject();
+		Writer.BeginObject();
+		TestData::WriteOp(Writer, "Op3", {}, {});
+		Writer.EndObject();
+	};
+	FOplogManifest ManifestWithArtifactsOrImports = TestData::GenerateCustomTestManifest(OpsWithArtifactsOrImportExports);
+	FOplogManifest ManifestWithMissingData = TestData::GenerateCustomTestManifest(OpsWithMissingArtifactsOrImportExports);
+	DiffBothDirections(ManifestWithArtifactsOrImports, ManifestWithMissingData, [](const FOplogDiffResults& Diff) 
+	{
+		CHECK(Diff.IdenticalOps.Num() == 0);
+		CHECK(Diff.OpsMissingInManifest1.Num() == 0);
+		CHECK(Diff.OpsMissingInManifest2.Num() == 0);
+		CHECK(Diff.ChangedOpsWithSameOutput.Num() == 3);
+		CHECK(Diff.ChangedOpsWithOutputDifferences.Num() == 0);
+		CHECK(ChangedOpPropertyExists(Diff, "Op1", "meta.cook.artifacts"));
+		CHECK(ChangedOpPropertyExists(Diff, "Op2", "meta.cook.importexport"));
+		CHECK(ChangedOpPropertyExists(Diff, "Op3", "meta.cook.artifacts"));
+		CHECK(ChangedOpPropertyExists(Diff, "Op3", "meta.cook.importexport"));
+	});
+}
+
+TEST_CASE("ZenOplogUtils::Diff::OpsWithMatchingPackages", "[ZenOplogUtils]")
+{
+	auto OpsWithPackages = [](FCbWriter& Writer)
+	{
+		Writer.BeginObject();
+		TestData::WriteOp(Writer, "Op1");
+		Writer.BeginArray("packagedata");
+		TestData::WritePackageData(Writer, FOplogManifest::FOp::FPackageData{
+			.ID = FCbObjectId(FCbObjectId::ByteArray{0}),
+			.Size = 20,
+			.RawSize = 32,
+			.DataHash = FIoHash(TEXT("aaaaaabbbbbbbcccccccddddddeeeeeeffffffff")),
+			.Filename = "Pack1"
+			});
+		Writer.EndArray();
+		Writer.EndObject();
+		Writer.BeginObject();
+		TestData::WriteOp(Writer, "Op2");
+		Writer.BeginArray("packagedata");
+		TestData::WritePackageData(Writer, FOplogManifest::FOp::FPackageData{
+			.ID = FCbObjectId(FCbObjectId::ByteArray{1}),
+			.Size = 30,
+			.RawSize = 40,
+			.DataHash = FIoHash(TEXT("aa1aaa3bbb5bbc7ccc8cddd3dde1eeeeffffffff")),
+			.Filename = "Pack2"
+			});
+		TestData::WritePackageData(Writer, FOplogManifest::FOp::FPackageData{
+			.ID = FCbObjectId(FCbObjectId::ByteArray{2}),
+			.Size = 40,
+			.RawSize = 50,
+			.DataHash = FIoHash(TEXT("aaaaaabb1bb33c3cccc3d3dd6d6ee7e7f7ffffff")),
+			.Filename = "Pack3"
+			});
+		Writer.EndArray();
+		Writer.EndObject();
+	};
+	FOplogManifest ManifestWithPackageEntries = TestData::GenerateCustomTestManifest(OpsWithPackages);
+	FOplogManifest OtherManifestWithPackageEntries = TestData::GenerateCustomTestManifest(OpsWithPackages);
+	DiffBothDirections(ManifestWithPackageEntries, OtherManifestWithPackageEntries, [](const FOplogDiffResults& Diff)
+	{
+		CHECK(Diff.IdenticalOps.Num() == 2);
+		CHECK(Diff.OpsMissingInManifest1.Num() == 0);
+		CHECK(Diff.OpsMissingInManifest2.Num() == 0);
+		CHECK(Diff.ChangedOpsWithSameOutput.Num() == 0);
+		CHECK(Diff.ChangedOpsWithOutputDifferences.Num() == 0);
+	});
+}
+
+TEST_CASE("ZenOplogUtils::Diff::OpsWithDifferentBulkDatas", "[ZenOplogUtils]")
+{
+	auto OpWithBulkData = [](FCbWriter& Writer)
+	{
+		const FOplogManifest::FOp::FBulkData TestBulkData1
+		{
+			.ID = FCbObjectId(FCbObjectId::ByteArray{0}),
+			.TypeStr = "Mmap",
+			.Size = 580,
+			.RawSize = 1400,
+			.DataHash = FIoHash(TEXT("aaaabbbbccccdddd4b0e73b1a0740277a82d732d")),
+			.Filename = "SomeFileWithLotsOfData.ubulk"
+		};
+		const FOplogManifest::FOp::FBulkData TestBulkData2
+		{
+			.ID = FCbObjectId(FCbObjectId::ByteArray{1}),
+			.TypeStr = "Mmap",
+			.Size = 31,
+			.RawSize = 20,
+			.DataHash = FIoHash(TEXT("eeeeefffffffaaaaaabbb3b1a07402aaaaa11111")),
+			.Filename = "AnotherFile.ubulk"
+		};
+		Writer.BeginObject();
+		TestData::WriteOp(Writer, "Op1");
+		Writer.BeginArray("bulkdata");
+		TestData::WriteBulkData(Writer, TestBulkData1);
+		TestData::WriteBulkData(Writer, TestBulkData2);
+		Writer.EndArray();
+		Writer.EndObject();
+	};
+	auto OpWithChangedBulkData = [](FCbWriter& Writer)
+	{
+		const FOplogManifest::FOp::FBulkData TestBulkData1
+		{
+			.ID = FCbObjectId(FCbObjectId::ByteArray{0}),
+			.TypeStr = "Mmap",
+			.Size = 580,
+			.RawSize = 1400,
+			.DataHash = FIoHash(TEXT("aaaabbbbccccdddd4b0e73b1a0740277a82d732d")),
+			.Filename = "SomeFileWithLotsOfData.ubulk"
+		};
+		const FOplogManifest::FOp::FBulkData TestBulkData2
+		{
+			.ID = FCbObjectId(FCbObjectId::ByteArray{8}),
+			.TypeStr = "SomethingElse",
+			.Size = 20,
+			.RawSize = 20,
+			.DataHash = FIoHash(TEXT("eeeeefff1234aaaaaabbb3b1a07402aaaaa11111")),
+			.Filename = "AnotherFileButNotTheSame.ubulk"
+		};
+		Writer.BeginObject();
+		TestData::WriteOp(Writer, "Op1");
+		Writer.BeginArray("bulkdata");
+		TestData::WriteBulkData(Writer, TestBulkData1);
+		TestData::WriteBulkData(Writer, TestBulkData2);
+		Writer.EndArray();
+		Writer.EndObject();
+	};
+	FOplogManifest Manifest1 = TestData::GenerateCustomTestManifest(OpWithBulkData);
+	FOplogManifest Manifest2 = TestData::GenerateCustomTestManifest(OpWithChangedBulkData);
+	DiffBothDirections(Manifest1, Manifest2, [](const FOplogDiffResults& Diff) 
+	{
+		CHECK(Diff.IdenticalOps.Num() == 0);
+		CHECK(Diff.OpsMissingInManifest1.Num() == 0);
+		CHECK(Diff.OpsMissingInManifest2.Num() == 0);
+		CHECK(Diff.ChangedOpsWithSameOutput.Num() == 0);
+		CHECK(Diff.ChangedOpsWithOutputDifferences.Num() == 1);
+		CHECK(ChangedOpPropertyExists(Diff, "Op1", "bulkdata/1/id"));
+		CHECK(ChangedOpPropertyExists(Diff, "Op1", "bulkdata/1/type"));
+		CHECK(ChangedOpPropertyExists(Diff, "Op1", "bulkdata/1/size"));
+		CHECK(ChangedOpPropertyExists(Diff, "Op1", "bulkdata/1/data"));
+		CHECK(ChangedOpPropertyExists(Diff, "Op1", "bulkdata/1/filename"));
+	});
+}
+
+TEST_CASE("ZenOplogUtils::Diff::OpsWithDifferentFiles", "[ZenOplogUtils]")
+{
+	const FOplogManifest::FOp::FFile TestFile1
+	{
+		.ID = FCbObjectId(FCbObjectId::ByteArray{0}),
+		.ServerPath = "Path//On//Server//1",
+		.ClientPath = "Path//On//Client//1",
+		.DataHash = FIoHash(TEXT("aaaabbbbccccddddeeeeeeffffff0277a82d732d"))
+	};
+	const FOplogManifest::FOp::FFile TestFile2
+	{
+		.ID = FCbObjectId(FCbObjectId::ByteArray{1}),
+		.ServerPath = "Path//On//Server//2",
+		.ClientPath = "Path//On//Client//2",
+		.DataHash = FIoHash(TEXT("aaaabb3bc666666666666e11ffff0277a82d732d"))
+	};
+	const FOplogManifest::FOp::FFile TestFile3
+	{
+		.ID = FCbObjectId(FCbObjectId::ByteArray{2}),
+		.ServerPath = "Path//On//Server//3",
+		.ClientPath = "Path//On//Client//3",
+		.DataHash = FIoHash(TEXT("aaaa1111111cddddeeee11ffffff0277a82d732d"))
+	};
+	const FOplogManifest::FOp::FFile TestFile4
+	{
+		.ID = FCbObjectId(FCbObjectId::ByteArray{3}),
+		.ServerPath = "Path//On//Server",
+		.ClientPath = "Path//On//Client",
+		.DataHash = FIoHash(TEXT("aaaab555555555555eee333333333333a82d732d"))
+	};
+	auto OpWithFiles1 = [&](FCbWriter& Writer)
+	{
+		Writer.BeginObject();
+		TestData::WriteOp(Writer, "Op1");
+		Writer.BeginArray("files");
+		TestData::WriteFile(Writer, TestFile1);
+		TestData::WriteFile(Writer, TestFile2);
+		Writer.EndArray();
+		Writer.EndObject();
+		Writer.BeginObject();
+		TestData::WriteOp(Writer, "Op2");
+		Writer.BeginArray("files");
+		TestData::WriteFile(Writer, TestFile3);
+		Writer.EndArray();
+		Writer.EndObject();
+		Writer.BeginObject();
+		TestData::WriteOp(Writer, "Op3");
+		Writer.BeginArray("files");
+		TestData::WriteFile(Writer, TestFile1);
+		TestData::WriteFile(Writer, TestFile3);
+		TestData::WriteFile(Writer, TestFile4);
+		Writer.EndArray();
+		Writer.EndObject();
+	};
+	auto OpWithFiles2 = [&](FCbWriter& Writer)
+	{
+		Writer.BeginObject();
+		TestData::WriteOp(Writer, "Op1");
+		Writer.BeginArray("files");
+		TestData::WriteFile(Writer, TestFile3);
+		Writer.EndArray();
+		Writer.EndObject();
+		Writer.BeginObject();
+		TestData::WriteOp(Writer, "Op2");
+		Writer.BeginArray("files");
+		TestData::WriteFile(Writer, TestFile1);
+		Writer.EndArray();
+		Writer.EndObject();
+		Writer.BeginObject();
+		TestData::WriteOp(Writer, "Op3");
+		Writer.BeginArray("files");
+		TestData::WriteFile(Writer, TestFile1);
+		TestData::WriteFile(Writer, TestFile2);
+		TestData::WriteFile(Writer, TestFile3);
+		TestData::WriteFile(Writer, TestFile4);
+		Writer.EndArray();
+		Writer.EndObject();
+	};
+	FOplogManifest Manifest1 = TestData::GenerateCustomTestManifest(OpWithFiles1);
+	FOplogManifest Manifest2 = TestData::GenerateCustomTestManifest(OpWithFiles2);
+	DiffBothDirections(Manifest1, Manifest2, [](const FOplogDiffResults& Diff) 
+	{
+		CHECK(Diff.IdenticalOps.Num() == 0);
+		CHECK(Diff.OpsMissingInManifest1.Num() == 0);
+		CHECK(Diff.OpsMissingInManifest2.Num() == 0);
+		CHECK(Diff.ChangedOpsWithSameOutput.Num() == 0);
+		CHECK(Diff.ChangedOpsWithOutputDifferences.Num() == 3);
+		CHECK(ChangedOpPropertyExists(Diff, "Op1", "files"));
+		CHECK(ChangedOpPropertyExists(Diff, "Op2", "files/0/id"));
+		CHECK(ChangedOpPropertyExists(Diff, "Op2", "files/0/serverpath"));
+		CHECK(ChangedOpPropertyExists(Diff, "Op2", "files/0/clientpath"));
+		CHECK(ChangedOpPropertyExists(Diff, "Op2", "files/0/data"));
+		CHECK(ChangedOpPropertyExists(Diff, "Op3", "files"));
+	});
+}
+
+TEST_CASE("ZenOplogUtils::Diff::DiffUnexpectedData", "[ZenOplogUtils]")
+{
+	auto OpsWithUnexpectedData1 = [](FCbWriter& Writer)
+	{
+		Writer.BeginObject();
+		TestData::WriteOp(Writer, "Op1");
+		Writer.AddString("somenewstring", "Blah");
+		Writer.AddInteger("somenewinteger", 32);
+		Writer.EndObject();
+		Writer.BeginObject();
+		TestData::WriteOp(Writer, "Op2");
+		Writer.BeginArray("SomeNewArray");
+		Writer.BeginObject();
+		Writer.AddBinaryAttachment("SomeNewAttachment", FIoHash(TEXT("aaaaaaaaabbbbbbbbbbddddddddddddddddddddd")));
+		Writer.AddString("FilenameOrSomething", "Test");
+		Writer.EndObject();
+		Writer.EndArray();
+		Writer.EndObject();
+	};
+	auto OpsWithUnexpectedData2 = [](FCbWriter& Writer)
+	{
+		Writer.BeginObject();
+		TestData::WriteOp(Writer, "Op1");
+		Writer.AddString("somenewstring", "Blah");
+		Writer.AddInteger("somenewinteger", 32);
+		Writer.EndObject();
+		Writer.BeginObject();
+		TestData::WriteOp(Writer, "Op2");
+		Writer.AddString("someotherstring", "Eh?");
+		Writer.AddString("adifferentstring", "Wat");
+		Writer.BeginArray("SomeNewArray");
+		Writer.BeginObject();
+		Writer.AddBinaryAttachment("SomeNewAttachment", FIoHash(TEXT("eeeeeeeeeeffffffffffffaaaaaaaaaaaaaddddd")));
+		Writer.AddString("FilenameOrSomething", "TestingThisThing");
+		Writer.EndObject();
+		Writer.EndArray();
+		Writer.EndObject();
+	};
+	FOplogManifest Manifest1 = TestData::GenerateCustomTestManifest(OpsWithUnexpectedData1);
+	FOplogManifest Manifest2 = TestData::GenerateCustomTestManifest(OpsWithUnexpectedData2);
+	DiffBothDirections(Manifest1, Manifest2, [](const FOplogDiffResults& Diff) 
+	{
+		CHECK(Diff.IdenticalOps.Num() == 1);
+		CHECK(Diff.OpsMissingInManifest1.Num() == 0);
+		CHECK(Diff.OpsMissingInManifest2.Num() == 0);
+		CHECK(Diff.ChangedOpsWithSameOutput.Num() == 1);
+		CHECK(Diff.ChangedOpsWithOutputDifferences.Num() == 0);
+		CHECK(ChangedOpPropertyExists(Diff, "Op2", "someotherstring"));
+		CHECK(ChangedOpPropertyExists(Diff, "Op2", "SomeNewArray/0/SomeNewAttachment"));
+		CHECK(ChangedOpPropertyExists(Diff, "Op2", "SomeNewArray/0/FilenameOrSomething"));
+		CHECK(ChangedOpPropertyExists(Diff, "Op2", "adifferentstring"));
+	});
+}
+
+TEST_CASE("ZenOplogUtils::Diff::Sorting", "[ZenOplogUtils]")
+{
+	auto Ops1 = [](FCbWriter& Writer)
+	{
+		Writer.BeginObject();
+		TestData::WriteOp(Writer, "OpA");
+		Writer.EndObject();
+		Writer.BeginObject();
+		TestData::WriteOp(Writer, "OpD");
+		Writer.AddString("SomeValue", "Blah");
+		Writer.EndObject();
+		Writer.BeginObject();
+		TestData::WriteOp(Writer, "OpC");
+		Writer.EndObject();
+		Writer.BeginObject();
+		TestData::WriteOp(Writer, "OpE");
+		Writer.EndObject();
+		Writer.BeginObject();
+		TestData::WriteOp(Writer, "OpB");
+		Writer.EndObject();
+		Writer.BeginObject();
+		TestData::WriteOp(Writer, "OpF");
+		Writer.AddString("SomeValue", "Blah2");
+		Writer.EndObject();
+		Writer.BeginObject();
+		TestData::WriteOp(Writer, "OpH");
+		Writer.EndObject();
+		Writer.BeginObject();
+		TestData::WriteOp(Writer, "OpG");
+		Writer.EndObject();
+	};
+	auto Ops2 = [](FCbWriter& Writer)
+	{
+		Writer.BeginObject();
+		TestData::WriteOp(Writer, "OpC");
+		Writer.EndObject();
+		Writer.BeginObject();
+		TestData::WriteOp(Writer, "OpA");
+		Writer.EndObject();
+		Writer.BeginObject();
+		TestData::WriteOp(Writer, "OpD");
+		Writer.AddString("SomeValue", "Eh");
+		Writer.EndObject();
+		Writer.BeginObject();
+		TestData::WriteOp(Writer, "OpF");
+		Writer.AddString("SomeValue", "Hmm");
+		Writer.EndObject();
+		Writer.BeginObject();
+		TestData::WriteOp(Writer, "OpH");
+		Writer.EndObject();
+		Writer.BeginObject();
+		TestData::WriteOp(Writer, "OpG");
+		Writer.EndObject();
+		Writer.BeginObject();
+		TestData::WriteOp(Writer, "OpJ");
+		Writer.EndObject();
+		Writer.BeginObject();
+		TestData::WriteOp(Writer, "OpI");
+		Writer.EndObject();
+	};
+	FOplogManifest Manifest1 = TestData::GenerateCustomTestManifest(Ops1);
+	FOplogManifest Manifest2 = TestData::GenerateCustomTestManifest(Ops2);
+	FOplogDiffResults Diff = DiffManifests(Manifest1, Manifest2);
+	CHECK(Diff.IdenticalOps.Num() == 4);
+	CHECK(Diff.OpsMissingInManifest1.Num() == 2);
+	CHECK(Diff.OpsMissingInManifest2.Num() == 2);
+	CHECK(Diff.ChangedOpsWithSameOutput.Num() == 2);
+	CHECK(Diff.ChangedOpsWithOutputDifferences.Num() == 0);
+
+	UE::SortDiffResults(Diff);
+	CHECK(Algo::IsSorted(Diff.IdenticalOps));
+	CHECK(Algo::IsSorted(Diff.OpsMissingInManifest1));
+	CHECK(Algo::IsSorted(Diff.OpsMissingInManifest2));
+	CHECK(Algo::IsSortedBy(Diff.ChangedOpsWithSameOutput, &FOplogDiffChangedOp::Key));
+	CHECK(Algo::IsSortedBy(Diff.ChangedOpsWithOutputDifferences, &FOplogDiffChangedOp::Key));
+}
+
+TEST_CASE("ZenOplogUtils::Diff::Helpers", "[ZenOplogUtils]")
+{
+	auto OpWithArrays1 = [](FCbWriter& Writer)
+	{
+		Writer.BeginObject();
+		TestData::WriteOp(Writer, "Op1");
+		Writer.BeginArray("testarray");
+		Writer.BeginObject();
+		Writer.AddBinaryAttachment("SomeNewAttachment", FIoHash(TEXT("aaaaaaaaabbbbbbbbbbddddddddddddddddddddd")));
+		Writer.AddString("FilenameOrSomething", "Test");
+		Writer.EndObject();
+		Writer.BeginObject();
+		Writer.AddBinaryAttachment("SomeNewAttachment", FIoHash(TEXT("bbbbbbbbbbbccccccccccccccddddddddeeeeeee")));
+		Writer.AddString("FilenameOrSomething", "AnotherTest");
+		Writer.EndObject();
+		Writer.EndArray();
+		Writer.EndObject();
+	};
+	auto OpWithArrays2 = [](FCbWriter& Writer)
+	{
+		Writer.BeginObject();
+		TestData::WriteOp(Writer, "Op1");
+		Writer.BeginArray("testarray");
+		Writer.BeginObject();
+		Writer.AddBinaryAttachment("SomeNewAttachment", FIoHash(TEXT("aaaaaaaaabbbbbbbbbbddddddddddddddddddddd")));
+		Writer.AddString("FilenameOrSomething", "DifferentValue");
+		Writer.EndObject();
+		Writer.BeginObject();
+		Writer.AddBinaryAttachment("SomeNewAttachment", FIoHash(TEXT("bbbbbbbbbbbccccccccccccccddddddddeeeeeee")));
+		Writer.AddString("FilenameOrSomething", "AnotherTest");
+		Writer.EndObject();
+		Writer.EndArray();
+		Writer.EndObject();
+	};
+	FOplogManifest Manifest1 = TestData::GenerateCustomTestManifest(OpWithArrays1);
+	FOplogManifest Manifest2 = TestData::GenerateCustomTestManifest(OpWithArrays2);
+	FOplogDiffResults Diff = DiffManifests(Manifest1, Manifest2);
+	CHECK(Diff.IdenticalOps.Num() == 0);
+	CHECK(Diff.OpsMissingInManifest1.Num() == 0);
+	CHECK(Diff.OpsMissingInManifest2.Num() == 0);
+	CHECK(Diff.ChangedOpsWithSameOutput.Num() == 1);
+	CHECK(Diff.ChangedOpsWithOutputDifferences.Num() == 0);
+	CHECK(Diff.ChangedOpsWithSameOutput[0].Value.Num() == 1);
+	CHECK(OpDiffContainsArrayValueChange(Diff.ChangedOpsWithSameOutput[0], "testarray", "FilenameOrSomething"));
+}
+
+TEST_CASE("ZenOplogUtils::Diff::DiffChangeHelper", "[ZenOplogUtils]")
+{
+	FOplogDiffChangedOp ChangedOp;
+	ChangedOp.Key = "Test";
+	ChangedOp.Value.Emplace(FCbEntryDifference{ .PropertyName = "Array1" });
+	ChangedOp.Value.Emplace(FCbEntryDifference{ .PropertyName = "Array2/0/Prop1" });
+	ChangedOp.Value.Emplace(FCbEntryDifference{ .PropertyName = "Array2/1/Prop2" });
+	ChangedOp.Value.Emplace(FCbEntryDifference{ .PropertyName = "Array2/2/SomeObject/Prop3" });
+	ChangedOp.Value.Emplace(FCbEntryDifference{ .PropertyName = "SomeObject/Array3/0/Prop4" });
+	ChangedOp.Value.Emplace(FCbEntryDifference{ .PropertyName = "SomeObject/Array4/0/SomeObject/Prop5" });
+	ChangedOp.Value.Emplace(FCbEntryDifference{ .PropertyName = "Array5/NotAnArrayEntry" });
+	ChangedOp.Value.Emplace(FCbEntryDifference{ .PropertyName = "Array5/NotAnArrayEntry/AlsoNotAnArrayEntry" });
+	CHECK(OpDiffContainsArrayValueChange(ChangedOp, "Array1", "SomeProperty"));
+	CHECK(OpDiffContainsArrayValueChange(ChangedOp, "Array2", "Prop1"));
+	CHECK(OpDiffContainsArrayValueChange(ChangedOp, "Array2", "Prop2"));
+	CHECK(!OpDiffContainsArrayValueChange(ChangedOp, "Array5", "NotAnArrayEntry"));
+	CHECK(!OpDiffContainsArrayValueChange(ChangedOp, "Array5", "AlsoNotAnArrayEntry"));
+	CHECK(!OpDiffContainsArrayValueChange(ChangedOp, "Array2", "Prop3"));
+	CHECK(OpDiffContainsArrayValueChange(ChangedOp, "Array2", "SomeObject/Prop3"));
+	CHECK(!OpDiffContainsArrayValueChange(ChangedOp, "Array3", "Prop4"));
+	CHECK(OpDiffContainsArrayValueChange(ChangedOp, "SomeObject/Array3", "Prop4"));
+	CHECK(!OpDiffContainsArrayValueChange(ChangedOp, "SomeObject/Array4", "Prop5"));
+	CHECK(OpDiffContainsArrayValueChange(ChangedOp, "SomeObject/Array4", "SomeObject/Prop5"));
+}
+
+} // UE
+#endif // WITH_EDITORONLY_DATA
+
+#endif // WITH_LOW_LEVEL_TESTS
diff --git a/Engine/Source/Developer/ZenOplogUtils/ZenOplogUtils.Build.cs b/Engine/Source/Developer/ZenOplogUtils/ZenOplogUtils.Build.cs
new file mode 100644
--- /dev/null
+++ b/Engine/Source/Developer/ZenOplogUtils/ZenOplogUtils.Build.cs
@@ -0,0 +1,14 @@
+// Copyright Epic Games, Inc. All Rights Reserved.
+
+using UnrealBuildTool;
+using System.IO;
+
+public class ZenOplogUtils : ModuleRules
+{
+	public ZenOplogUtils(ReadOnlyTargetRules Target) : base(Target)
+	{
+		PrivateDependencyModuleNames.AddRange(new string[] { "Core", "Zen", "Json" });
+		PrivateIncludePathModuleNames.AddRange(new string[] { "DesktopPlatform" });
+		CppCompileWarningSettings.UnsafeTypeCastWarningLevel = WarningLevel.Error;
+	}
+}
diff --git a/Engine/Source/Programs/LowLevelTests/FoundationTests/FoundationTests.Build.cs b/Engine/Source/Programs/LowLevelTests/FoundationTests/FoundationTests.Build.cs
--- a/Engine/Source/Programs/LowLevelTests/FoundationTests/FoundationTests.Build.cs
+++ b/Engine/Source/Programs/LowLevelTests/FoundationTests/FoundationTests.Build.cs
@@ -51,6 +51,7 @@ public class FoundationTests : TestModuleRules
 				"TelemetryUtils",
 				"AssetRegistry",
 				"Serialization",
+				"ZenOplogUtils"
 			});
 
 		PrivateIncludePathModuleNames.AddRange(new string[] { 
diff --git a/Engine/Source/Programs/ZenOplogDiffTool/Private/DownloadSnapshotOperation.cpp b/Engine/Source/Programs/ZenOplogDiffTool/Private/DownloadSnapshotOperation.cpp
new file mode 100644
--- /dev/null
+++ b/Engine/Source/Programs/ZenOplogDiffTool/Private/DownloadSnapshotOperation.cpp
@@ -0,0 +1,28 @@
+// Copyright Epic Games, Inc. All Rights Reserved.
+
+#include "DownloadSnapshotOperation.h"
+#include "ZenOplogDiffLogging.h"
+
+FDownloadSnapshot::FDownloadSnapshot(const UE::FZenSnapshotDescriptor& Descriptor, bool DownloadAsJson, FStringView OutputRoot)
+	: SnapshotDescriptor(Descriptor)
+	, OutputDirectory(OutputRoot)
+	, DownloadManifestAsJson(DownloadAsJson)
+{
+}
+
+IOplogDiffOperation::ERunningState FDownloadSnapshot::Run()
+{
+	UE::FDownloadOplogManifestResult DownloadResult = DownloadOplogManifest(SnapshotDescriptor, OutputDirectory, DownloadManifestAsJson);
+	if (DownloadResult.Result == UE::FDownloadOplogManifestResult::EStatus::Error)
+	{
+		UE_LOG(LogZenOplogDiffTool, Error, TEXT("Failed to download an oplog manifest!"));
+		return ERunningState::Error;
+	}
+	else
+	{
+		DownloadedManifest = DownloadResult.DownloadedFilename;
+		UE_LOG(LogZenOplogDiffTool, Display, TEXT("Snapshot downloaded to \"%s\""), *DownloadResult.DownloadedFilename);
+	}
+
+	return ERunningState::Success;
+}
diff --git a/Engine/Source/Programs/ZenOplogDiffTool/Private/DownloadSnapshotOperation.h b/Engine/Source/Programs/ZenOplogDiffTool/Private/DownloadSnapshotOperation.h
new file mode 100644
--- /dev/null
+++ b/Engine/Source/Programs/ZenOplogDiffTool/Private/DownloadSnapshotOperation.h
@@ -0,0 +1,21 @@
+// Copyright Epic Games, Inc. All Rights Reserved.
+
+#pragma once
+
+#include "IZenOplogDiffOperation.h"
+#include "Experimental/ZenOplogSnapshots.h"
+
+// Downloads a snapshot from zen, using a snapshot descriptor
+class FDownloadSnapshot : public IOplogDiffOperation
+{
+public:
+	FDownloadSnapshot(const UE::FZenSnapshotDescriptor& Descriptor, bool DownloadAsJson, FStringView OutputRoot);
+	virtual ~FDownloadSnapshot() = default;
+	virtual ERunningState Run() override;
+
+	UE::FZenSnapshotDescriptor SnapshotDescriptor;	// Used to identify the oplog to download
+	FString OutputDirectory;			// Directory to write the manifests to (optional)
+	bool DownloadManifestAsJson = false;
+
+	FString DownloadedManifest;		// Local path to the manifest that was successfully downloaded
+};
diff --git a/Engine/Source/Programs/ZenOplogDiffTool/Private/FindSnapshotOperation.cpp b/Engine/Source/Programs/ZenOplogDiffTool/Private/FindSnapshotOperation.cpp
new file mode 100644
--- /dev/null
+++ b/Engine/Source/Programs/ZenOplogDiffTool/Private/FindSnapshotOperation.cpp
@@ -0,0 +1,31 @@
+// Copyright Epic Games, Inc. All Rights Reserved.
+
+#include "FindSnapshotOperation.h"
+#include "ZenOplogDiffLogging.h"
+
+IOplogDiffOperation::ERunningState FFindSnapshot::Run()
+{
+	if (!PendingResult.IsValid())
+	{
+		PendingResult = UE::FindSnapshotDescriptorForBuild(FindProjectName, FindStream, FindPlatform, FindCommitID, bDisplayAllBuildsFound);
+	}
+	else if(PendingResult.IsReady())
+	{
+		FoundSnapshot = PendingResult.Get();
+		if (FoundSnapshot)
+		{
+			UE_LOG(LogZenOplogDiffTool, Display, TEXT("Found matching build:\nHost: %s\nNamespace: %s\nBucket: %s\nBuildID: %s"),
+				*FoundSnapshot->CloudHost,
+				*FoundSnapshot->Namespace,
+				*FoundSnapshot->Bucket,
+				*FoundSnapshot->BuildID);
+			return ERunningState::Success;
+		}
+		else
+		{
+			UE_LOG(LogZenOplogDiffTool, Error, TEXT("Failed to find a matching build"));
+			return ERunningState::Error;
+		}
+	}
+	return ERunningState::Running;
+}
diff --git a/Engine/Source/Programs/ZenOplogDiffTool/Private/FindSnapshotOperation.h b/Engine/Source/Programs/ZenOplogDiffTool/Private/FindSnapshotOperation.h
new file mode 100644
--- /dev/null
+++ b/Engine/Source/Programs/ZenOplogDiffTool/Private/FindSnapshotOperation.h
@@ -0,0 +1,24 @@
+// Copyright Epic Games, Inc. All Rights Reserved.
+
+#pragma once
+
+#include "IZenOplogDiffOperation.h"
+#include "Experimental/ZenOplogSnapshots.h"
+
+// Uses the zen build service to lookup a snapshot descriptor for a particular build
+class FFindSnapshot : public IOplogDiffOperation
+{
+public:
+	virtual ~FFindSnapshot() = default;
+	virtual ERunningState Run() override;
+
+	// Find a snapshot based on these parameters
+	FString FindProjectName;
+	FString FindStream;
+	FString FindPlatform;
+	FString FindCommitID;
+	bool bDisplayAllBuildsFound = false;
+
+	TSharedFuture<TOptional<UE::FZenSnapshotDescriptor>> PendingResult;
+	TOptional<UE::FZenSnapshotDescriptor> FoundSnapshot;
+};
diff --git a/Engine/Source/Programs/ZenOplogDiffTool/Private/IZenOplogDiffOperation.h b/Engine/Source/Programs/ZenOplogDiffTool/Private/IZenOplogDiffOperation.h
new file mode 100644
--- /dev/null
+++ b/Engine/Source/Programs/ZenOplogDiffTool/Private/IZenOplogDiffOperation.h
@@ -0,0 +1,23 @@
+// Copyright Epic Games, Inc. All Rights Reserved.
+
+#pragma once
+
+#include "CoreTypes.h"
+#include "Templates/Function.h"
+
+// Represents an operation that can run continuously.
+class IOplogDiffOperation
+{
+public:
+	virtual ~IOplogDiffOperation() = default;
+
+	enum class ERunningState
+	{
+		Running,
+		Success,
+		Error
+	};
+
+	virtual ERunningState Run() = 0;					// Continuously called until it returns Success or Error
+	TFunction<void(IOplogDiffOperation&)> OnComplete;	// Parameter = reference to operation that just completed. Should be called when Run returns Success
+};
diff --git a/Engine/Source/Programs/ZenOplogDiffTool/Private/LoadSnapshotDescriptorOperation.cpp b/Engine/Source/Programs/ZenOplogDiffTool/Private/LoadSnapshotDescriptorOperation.cpp
new file mode 100644
--- /dev/null
+++ b/Engine/Source/Programs/ZenOplogDiffTool/Private/LoadSnapshotDescriptorOperation.cpp
@@ -0,0 +1,16 @@
+// Copyright Epic Games, Inc. All Rights Reserved.
+
+#include "LoadSnapshotDescriptorOperation.h"
+#include "ZenOplogDiffLogging.h"
+
+FLoadSnapshotDescriptor::FLoadSnapshotDescriptor(FStringView SnapshotJson)
+	: SnapshotJsonPath(SnapshotJson)
+{
+}
+
+IOplogDiffOperation::ERunningState FLoadSnapshotDescriptor::Run()
+{
+	UE_LOG(LogZenOplogDiffTool, Verbose, TEXT("Parsing snapshot descriptors from %s"), *SnapshotJsonPath);
+	LoadedSnapshots = UE::ParseSnapshotDescriptorFromJson(SnapshotJsonPath);
+	return LoadedSnapshots.Num() > 0 ? ERunningState::Success : ERunningState::Error;
+}
diff --git a/Engine/Source/Programs/ZenOplogDiffTool/Private/LoadSnapshotDescriptorOperation.h b/Engine/Source/Programs/ZenOplogDiffTool/Private/LoadSnapshotDescriptorOperation.h
new file mode 100644
--- /dev/null
+++ b/Engine/Source/Programs/ZenOplogDiffTool/Private/LoadSnapshotDescriptorOperation.h
@@ -0,0 +1,18 @@
+// Copyright Epic Games, Inc. All Rights Reserved.
+
+#pragma once
+
+#include "IZenOplogDiffOperation.h"
+#include "Experimental/ZenOplogSnapshots.h"
+
+// Load a snapshot descriptor from horde build artifact json
+class FLoadSnapshotDescriptor : public IOplogDiffOperation
+{
+public:
+	FLoadSnapshotDescriptor(FStringView SnapshotJson);
+	virtual ~FLoadSnapshotDescriptor() = default;
+	virtual ERunningState Run() override;
+
+	FString SnapshotJsonPath;		// Path to the file to load
+	TArray<UE::FZenSnapshotDescriptor> LoadedSnapshots;		// Snapshot json may contain more than one build
+};
diff --git a/Engine/Source/Programs/ZenOplogDiffTool/Private/OplogManifestDiffOperation.cpp b/Engine/Source/Programs/ZenOplogDiffTool/Private/OplogManifestDiffOperation.cpp
new file mode 100644
--- /dev/null
+++ b/Engine/Source/Programs/ZenOplogDiffTool/Private/OplogManifestDiffOperation.cpp
@@ -0,0 +1,62 @@
+// Copyright Epic Games, Inc. All Rights Reserved.
+
+#include "OplogManifestDiffOperation.h"
+#include "ZenOplogDiffLogging.h"
+#include "Experimental/DiffCompactBinary.h"
+#include "Experimental/ZenOplogDiff.h"
+#include "Async/TaskGraphInterfaces.h"
+
+FOplogManifestDiff::FOplogManifestDiff(FString Manifest1, FString Manifest2)
+	: ManifestToDiff1(Manifest1)
+	, ManifestToDiff2(Manifest2)
+{
+}
+
+IOplogDiffOperation::ERunningState FOplogManifestDiff::Run()
+{
+	if (!LoadManifests())
+	{
+		UE_LOG(LogZenOplogDiffTool, Error, TEXT("Failed to load manifests"));
+		return ERunningState::Error;
+	}
+	UE_LOG(LogZenOplogDiffTool, Display, TEXT("Diffing oplogs..."));
+	DiffResults = DiffManifests(Manifest1, Manifest2);
+
+	return ERunningState::Success;
+}
+
+bool FOplogManifestDiff::LoadManifests()
+{
+	UE_LOG(LogZenOplogDiffTool, Display, TEXT("Loading oplogs..."));
+	FGraphEventArray LoadingTasks;
+	TAtomic<bool> bLoadedOk = true;
+	LoadingTasks.Add(FFunctionGraphTask::CreateAndDispatchWhenReady([this, &bLoadedOk]()
+	{
+		UE::FLoadOplogManifestResult LoadedManifest1 = UE::LoadOplogManifestFromFile(ManifestToDiff1);
+		if (LoadedManifest1.Result != UE::FLoadOplogManifestResult::EStatus::Ok)
+		{
+			UE_LOG(LogZenOplogDiffTool, Error, TEXT("Failed to load manifest from file %s"), *ManifestToDiff1);
+			bLoadedOk = false;
+		}
+		else
+		{
+			Manifest1 = MoveTemp(*LoadedManifest1.Manifest);
+		}
+	}));
+	LoadingTasks.Add(FFunctionGraphTask::CreateAndDispatchWhenReady([this, &bLoadedOk]()
+	{
+		UE::FLoadOplogManifestResult LoadedManifest2 = UE::LoadOplogManifestFromFile(ManifestToDiff2);
+		if (LoadedManifest2.Result != UE::FLoadOplogManifestResult::EStatus::Ok)
+		{
+			UE_LOG(LogZenOplogDiffTool, Error, TEXT("Failed to load manifest from file %s"), *ManifestToDiff2);
+			bLoadedOk = false;
+		}
+		else
+		{
+			Manifest2 = MoveTemp(*LoadedManifest2.Manifest);
+		}
+	}));
+	FTaskGraphInterface::Get().WaitUntilTasksComplete(LoadingTasks);
+
+	return bLoadedOk;
+}
diff --git a/Engine/Source/Programs/ZenOplogDiffTool/Private/OplogManifestDiffOperation.h b/Engine/Source/Programs/ZenOplogDiffTool/Private/OplogManifestDiffOperation.h
new file mode 100644
--- /dev/null
+++ b/Engine/Source/Programs/ZenOplogDiffTool/Private/OplogManifestDiffOperation.h
@@ -0,0 +1,28 @@
+// Copyright Epic Games, Inc. All Rights Reserved.
+
+#pragma once
+
+#include "IZenOplogDiffOperation.h"
+#include "Experimental/ZenOplogManifest.h"
+#include "Experimental/ZenOplogDiff.h"
+
+class FOplogManifestDiff : public IOplogDiffOperation
+{
+public:
+	FOplogManifestDiff(FString Manifest1, FString Manifest2);
+	virtual ~FOplogManifestDiff() = default;
+	virtual ERunningState Run() override;
+
+	// Paths to local manifest files
+	FString ManifestToDiff1;
+	FString ManifestToDiff2;
+
+	UE::FOplogManifest Manifest1;
+	UE::FOplogManifest Manifest2;
+
+	// Results of the diff
+	UE::FOplogDiffResults DiffResults;
+
+private:
+	bool LoadManifests();
+};
diff --git a/Engine/Source/Programs/ZenOplogDiffTool/Private/OutputDiffResultsOperation.cpp b/Engine/Source/Programs/ZenOplogDiffTool/Private/OutputDiffResultsOperation.cpp
new file mode 100644
--- /dev/null
+++ b/Engine/Source/Programs/ZenOplogDiffTool/Private/OutputDiffResultsOperation.cpp
@@ -0,0 +1,228 @@
+// Copyright Epic Games, Inc. All Rights Reserved.
+
+#include "OutputDiffResultsOperation.h"
+#include "ZenOplogDiffLogging.h"
+#include "Experimental/DiffCompactBinary.h"
+#include "Algo/Find.h"
+#include "Async/ParallelFor.h"
+#include "HAL/PlatformFileManager.h"
+#include "Misc/FileHelper.h"
+#include "Misc/Paths.h"
+#include "Serialization/CompactBinaryWriter.h"
+#include "Serialization/JsonWriter.h"
+
+struct FOutputDiffResults::FDiffSummary
+{
+	int64 IdenticalPackageDataTotalSize = 0;
+	int64 IdenticalBulkDataTotalSize = 0;
+	int64 NewPackageDataTotalSize = 0;
+	int64 NewBulkDataTotalSize = 0;
+	int64 DeletedPackageDataTotalSize = 0;
+	int64 DeletedBulkDataTotalSize = 0;
+	int64 ChangedPackageDataTotalSize = 0;
+	int64 ChangedBulkDataTotalSize = 0;
+
+	void Output(TSharedRef<TJsonWriter<>>& Writer) const
+	{
+		Writer->WriteObjectStart("Summary");
+		Writer->WriteValue("IdenticalPackageDataSize", IdenticalPackageDataTotalSize);
+		Writer->WriteValue("IdenticalBulkDataSize", IdenticalBulkDataTotalSize);
+		Writer->WriteValue("NewPackageDataTotalSize", NewPackageDataTotalSize);
+		Writer->WriteValue("NewBulkDataTotalSize", NewBulkDataTotalSize);
+		Writer->WriteValue("DeletedPackageDataTotalSize", DeletedPackageDataTotalSize);
+		Writer->WriteValue("DeletedBulkDataTotalSize", DeletedBulkDataTotalSize);
+		Writer->WriteValue("ChangedPackageDataTotalSize", ChangedPackageDataTotalSize);
+		Writer->WriteValue("ChangedBulkDataTotalSize", ChangedBulkDataTotalSize);
+		Writer->WriteObjectEnd();
+	}
+
+	void Output(FCbWriter& Writer) const
+	{
+		Writer.BeginObject("Summary");
+		Writer.AddInteger("IdenticalPackageDataSize", IdenticalPackageDataTotalSize);
+		Writer.AddInteger("IdenticalBulkDataSize", IdenticalBulkDataTotalSize);
+		Writer.AddInteger("NewPackageDataTotalSize", NewPackageDataTotalSize);
+		Writer.AddInteger("NewBulkDataTotalSize", NewBulkDataTotalSize);
+		Writer.AddInteger("DeletedPackageDataTotalSize", DeletedPackageDataTotalSize);
+		Writer.AddInteger("DeletedBulkDataTotalSize", DeletedBulkDataTotalSize);
+		Writer.AddInteger("ChangedPackageDataTotalSize", ChangedPackageDataTotalSize);
+		Writer.AddInteger("ChangedBulkDataTotalSize", ChangedBulkDataTotalSize);
+		Writer.EndObject();
+	}
+
+	void operator+=(const FDiffSummary& Other)
+	{
+		IdenticalPackageDataTotalSize += Other.IdenticalPackageDataTotalSize;
+		IdenticalBulkDataTotalSize += Other.IdenticalBulkDataTotalSize;
+		NewPackageDataTotalSize += Other.NewPackageDataTotalSize;
+		NewBulkDataTotalSize += Other.NewBulkDataTotalSize;
+		DeletedPackageDataTotalSize += Other.DeletedPackageDataTotalSize;
+		DeletedBulkDataTotalSize += Other.DeletedBulkDataTotalSize;
+		ChangedPackageDataTotalSize += Other.ChangedPackageDataTotalSize;
+		ChangedBulkDataTotalSize += Other.ChangedBulkDataTotalSize;
+	}
+};
+
+FOutputDiffResults::FOutputDiffResults(UE::FOplogManifest&& Oplog1, UE::FOplogManifest&& Oplog2, UE::FOplogDiffResults&& Results)
+	: Manifest1(MoveTemp(Oplog1))
+	, Manifest2(MoveTemp(Oplog2))
+	, DiffResults(MoveTemp(Results))
+{
+}
+
+void FOutputDiffResults::LogOutput()
+{
+	const bool bDisplaySummary = LoggerOptions == ELogOptions::Summary;
+	if (bDisplaySummary)
+	{
+		UE_LOG(LogZenOplogDiffTool, Display, TEXT("Identical Ops: %d"), DiffResults.IdenticalOps.Num());
+		UE_LOG(LogZenOplogDiffTool, Display, TEXT("New Ops: %d"), DiffResults.OpsMissingInManifest1.Num());
+		UE_LOG(LogZenOplogDiffTool, Display, TEXT("Deleted Ops: %d"), DiffResults.OpsMissingInManifest2.Num());
+		UE_LOG(LogZenOplogDiffTool, Display, TEXT("Changed Ops With Different Output: %d"), DiffResults.ChangedOpsWithOutputDifferences.Num());
+		UE_LOG(LogZenOplogDiffTool, Display, TEXT("Changed Ops With Same: %d"), DiffResults.ChangedOpsWithSameOutput.Num());
+	}
+}
+
+FOutputDiffResults::FDiffSummary FOutputDiffResults::SummarizeDiffResults()
+{
+	FDiffSummary TotalSummary;
+	auto SimpleSummary = [](const UE::FOplogManifest::FOp& TheOp, int64& PackageDataSize, int64& BulkDataSize)
+	{
+		for (const UE::FOplogManifest::FOp::FPackageData& PackageData : TheOp.GetPackageDatas())
+		{
+			PackageDataSize += PackageData.Size;
+		}
+		for (const UE::FOplogManifest::FOp::FBulkData& BulkData : TheOp.GetBulkDatas())
+		{
+			BulkDataSize += BulkData.Size;
+		}
+	};
+
+	TArray<FDiffSummary> Summaries;
+	ParallelForWithTaskContext(Summaries, DiffResults.IdenticalOps.Num(), [&SimpleSummary, this](FDiffSummary& Context, int32 Index)
+	{
+		const UE::FOplogManifest::FOp& TheOp = Manifest1.Ops[Manifest1.OpKeyToIndex[DiffResults.IdenticalOps[Index]]];
+		SimpleSummary(TheOp, Context.IdenticalPackageDataTotalSize, Context.IdenticalBulkDataTotalSize);
+	});
+	for (const FDiffSummary& Summary : Summaries)
+	{
+		TotalSummary += Summary;
+	}
+	ParallelForWithTaskContext(Summaries, DiffResults.ChangedOpsWithSameOutput.Num(), [&SimpleSummary, this](FDiffSummary& Context, int32 Index)
+	{
+		const UE::FOplogManifest::FOp& TheOp = Manifest1.Ops[Manifest1.OpKeyToIndex[DiffResults.ChangedOpsWithSameOutput[Index].Key]];
+		SimpleSummary(TheOp, Context.IdenticalPackageDataTotalSize, Context.IdenticalBulkDataTotalSize);
+	});
+	for (const FDiffSummary& Summary : Summaries)
+	{
+		TotalSummary += Summary;
+	}
+	ParallelForWithTaskContext(Summaries, DiffResults.OpsMissingInManifest1.Num(), [&SimpleSummary, this](FDiffSummary& Context, int32 Index)
+	{
+		const UE::FOplogManifest::FOp& TheOp = Manifest2.Ops[Manifest2.OpKeyToIndex[DiffResults.OpsMissingInManifest1[Index]]];
+		SimpleSummary(TheOp, Context.NewPackageDataTotalSize, Context.NewBulkDataTotalSize);
+	});
+	for (const FDiffSummary& Summary : Summaries)
+	{
+		TotalSummary += Summary;
+	}
+	ParallelForWithTaskContext(Summaries, DiffResults.OpsMissingInManifest2.Num(), [&SimpleSummary, this](FDiffSummary& Context, int32 Index)
+	{
+		const UE::FOplogManifest::FOp& TheOp = Manifest1.Ops[Manifest1.OpKeyToIndex[DiffResults.OpsMissingInManifest2[Index]]];
+		SimpleSummary(TheOp, Context.DeletedPackageDataTotalSize, Context.DeletedBulkDataTotalSize);
+	});
+	for (const FDiffSummary& Summary : Summaries)
+	{
+		TotalSummary += Summary;
+	}
+	ParallelForWithTaskContext(Summaries, DiffResults.ChangedOpsWithOutputDifferences.Num(), [&SimpleSummary, this](FDiffSummary& Context, int32 Index)
+	{
+		const UE::FOplogManifest::FOp& TheOp = Manifest2.Ops[Manifest2.OpKeyToIndex[DiffResults.ChangedOpsWithOutputDifferences[Index].Key]];
+		SimpleSummary(TheOp, Context.ChangedPackageDataTotalSize, Context.ChangedBulkDataTotalSize);
+	});
+	for (const FDiffSummary& Summary : Summaries)
+	{
+		TotalSummary += Summary;
+	}
+	return TotalSummary;
+}
+
+bool FOutputDiffResults::OutputCompactBinary(const FDiffSummary& Summary, UE::EOutputManifestDiffOptions Flags)
+{
+	FCbWriter Writer;
+	Writer.BeginObject();
+	OutputManifestDiffResultsToCompactBinary(Writer, DiffResults, Flags);
+	Summary.Output(Writer);
+	Writer.EndObject();
+	TUniquePtr<FArchive> Ar(IFileManager::Get().CreateFileWriter(*OutputPath));
+	if (Ar && !Ar->IsError())
+	{
+		Writer.Save(*Ar);
+		return !Ar->IsError();
+	}
+	else
+	{
+		return false;
+	}
+}
+
+bool FOutputDiffResults::OutputJson(const FDiffSummary& Summary, UE::EOutputManifestDiffOptions Flags)
+{
+	FString OutputString;
+	TSharedRef<TJsonWriter<>> JsonWriter = TJsonWriterFactory <>::Create(&OutputString, 0);
+	JsonWriter->WriteObjectStart();
+	OutputManifestDiffResultsToJson(JsonWriter, DiffResults, Flags);
+	Summary.Output(JsonWriter);
+	JsonWriter->WriteObjectEnd();
+	JsonWriter->Close();
+	if (!FFileHelper::SaveStringToFile(OutputString, *OutputPath))
+	{
+		UE_LOG(LogZenOplogDiffTool, Error, TEXT("Failed to write diff results to '%s'"), *OutputPath);
+		return false;
+	}
+	return true;
+}
+
+IOplogDiffOperation::ERunningState FOutputDiffResults::Run()
+{
+	UE_LOG(LogZenOplogDiffTool, Display, TEXT("Sorting diff results"));
+	UE::SortDiffResults(DiffResults);
+	FDiffSummary DiffSummary = SummarizeDiffResults();
+
+	LogOutput();
+	
+	if (OutputPath.Len() > 0)
+	{
+		const FString FileExtension = FPaths::GetExtension(OutputPath).ToLower();
+		UE::EOutputManifestDiffOptions OutputFlags = UE::EOutputManifestDiffOptions::None;
+		if (OutputOptions == EOutputOptions::OutputChangedValues)
+		{
+			OutputFlags |= UE::EOutputManifestDiffOptions::OutputDifferences;
+		}
+		
+		bool WriteSucceeded = false;
+		if (FileExtension == TEXT("json"))
+		{
+			UE_LOG(LogZenOplogDiffTool, Display, TEXT("Writing results to %s"), *OutputPath);
+			WriteSucceeded = OutputJson(DiffSummary, OutputFlags);
+		}
+		else if (FileExtension == TEXT("cb"))
+		{
+			UE_LOG(LogZenOplogDiffTool, Display, TEXT("Writing results to %s"), *OutputPath);
+			WriteSucceeded = OutputCompactBinary(DiffSummary, OutputFlags);
+		}
+		else
+		{
+			UE_LOG(LogZenOplogDiffTool, Error, TEXT("Unsupported diff output path '%s'. Expected .cb or .json extension"), *OutputPath);
+			return ERunningState::Error;
+		}
+
+		if(!WriteSucceeded)
+		{
+			UE_LOG(LogZenOplogDiffTool, Warning, TEXT("Failed to write results to file"));
+			return ERunningState::Error;
+		}
+	}
+	
+	return ERunningState::Success;
+}
diff --git a/Engine/Source/Programs/ZenOplogDiffTool/Private/OutputDiffResultsOperation.h b/Engine/Source/Programs/ZenOplogDiffTool/Private/OutputDiffResultsOperation.h
new file mode 100644
--- /dev/null
+++ b/Engine/Source/Programs/ZenOplogDiffTool/Private/OutputDiffResultsOperation.h
@@ -0,0 +1,44 @@
+// Copyright Epic Games, Inc. All Rights Reserved.
+
+#pragma once
+
+#include "IZenOplogDiffOperation.h"
+#include "Experimental/ZenOplogDiff.h"
+#include "Experimental/ZenOplogManifest.h"
+
+// Takes ownership of diff results + outputs them to a file and/or the log
+class FOutputDiffResults : public IOplogDiffOperation
+{
+public:
+	enum class ELogOptions
+	{
+		None,
+		Summary
+	};
+
+	enum class EOutputOptions
+	{
+		None,
+		OutputChangedValues
+	};
+
+	FOutputDiffResults(UE::FOplogManifest&& Oplog1, UE::FOplogManifest&& Oplog2, UE::FOplogDiffResults&& DiffResults);
+	virtual ~FOutputDiffResults() = default;
+	virtual ERunningState Run() override;
+
+	FString OutputPath;
+	ELogOptions LoggerOptions = ELogOptions::None;
+	EOutputOptions OutputOptions = EOutputOptions::None;
+
+	UE::FOplogManifest Manifest1;
+	UE::FOplogManifest Manifest2;
+	UE::FOplogDiffResults DiffResults;
+
+private:
+	struct FDiffSummary;
+
+	void LogOutput();
+	bool OutputCompactBinary(const FDiffSummary& Summary, UE::EOutputManifestDiffOptions Flags);
+	bool OutputJson(const FDiffSummary& Summary, UE::EOutputManifestDiffOptions Flags);
+	FDiffSummary SummarizeDiffResults();
+};
diff --git a/Engine/Source/Programs/ZenOplogDiffTool/Private/ZenOplogDiffLogging.h b/Engine/Source/Programs/ZenOplogDiffTool/Private/ZenOplogDiffLogging.h
new file mode 100644
--- /dev/null
+++ b/Engine/Source/Programs/ZenOplogDiffTool/Private/ZenOplogDiffLogging.h
@@ -0,0 +1,7 @@
+// Copyright Epic Games, Inc. All Rights Reserved.
+
+#pragma once 
+
+#include "Logging/LogMacros.h"
+
+DECLARE_LOG_CATEGORY_EXTERN(LogZenOplogDiffTool, Log, All);
diff --git a/Engine/Source/Programs/ZenOplogDiffTool/Private/ZenOplogDiffTool.cpp b/Engine/Source/Programs/ZenOplogDiffTool/Private/ZenOplogDiffTool.cpp
new file mode 100644
--- /dev/null
+++ b/Engine/Source/Programs/ZenOplogDiffTool/Private/ZenOplogDiffTool.cpp
@@ -0,0 +1,268 @@
+// Copyright Epic Games, Inc. All Rights Reserved.
+
+#include "RequiredProgramMainCPPInclude.h"
+#include "ZenOplogDiffLogging.h"
+#include "IZenOplogDiffOperation.h"
+#include "DownloadSnapshotOperation.h"
+#include "OplogManifestDiffOperation.h"
+#include "OutputDiffResultsOperation.h"
+#include "FindSnapshotOperation.h"
+#include "LoadSnapshotDescriptorOperation.h"
+#include "CoreGlobals.h"
+#include "HAL/LowLevelMemTracker.h"
+#include "Misc/Parse.h"
+#include "Modules/ModuleManager.h"
+#include "Experimental/ZenOplogDiff.h"
+
+DEFINE_LOG_CATEGORY(LogZenOplogDiffTool);
+
+IMPLEMENT_APPLICATION(ZenOplogDiffTool, "ZenOplogDiffTool");
+
+void PrintHelp()
+{
+	UE_LOG(LogZenOplogDiffTool, Display, TEXT("Zen Oplog Diff Utility"));
+	UE_LOG(LogZenOplogDiffTool, Display, TEXT("Usage:"));
+	UE_LOG(LogZenOplogDiffTool, Display, TEXT("-find-snapshot - Find a snapshot using the build service"));
+	UE_LOG(LogZenOplogDiffTool, Display, TEXT("\t-project=<name> - Project Name"));
+	UE_LOG(LogZenOplogDiffTool, Display, TEXT("\t-stream=<name> - Stream Name"));
+	UE_LOG(LogZenOplogDiffTool, Display, TEXT("\t-platform=<name> - Platform name"));
+	UE_LOG(LogZenOplogDiffTool, Display, TEXT("\t-commit=<cl> - Commit ID / Changelist"));
+	UE_LOG(LogZenOplogDiffTool, Display, TEXT("\t-displaybuilds - Output the list of namespaces, buckets, and commits found (optional)"));
+	UE_LOG(LogZenOplogDiffTool, Display, TEXT("-download-snapshot - Download an oplog snapshot from cloud storage"));
+	UE_LOG(LogZenOplogDiffTool, Display, TEXT("\t-output-directory=<path> - Base directory to download to (optional)"));
+	UE_LOG(LogZenOplogDiffTool, Display, TEXT("\t-as-json - Download the manifest as json (optional)"));
+	UE_LOG(LogZenOplogDiffTool, Display, TEXT("\t-snapshot-json=<path> - Path to an exported zen snapshot descriptor (from horde)"));
+	UE_LOG(LogZenOplogDiffTool, Display, TEXT("\t\tOR"));
+	UE_LOG(LogZenOplogDiffTool, Display, TEXT("\t-project=<name> - Project Name"));
+	UE_LOG(LogZenOplogDiffTool, Display, TEXT("\t-stream=<name> - Stream Name"));
+	UE_LOG(LogZenOplogDiffTool, Display, TEXT("\t-platform=<name> - Platform name"));
+	UE_LOG(LogZenOplogDiffTool, Display, TEXT("\t-commit=<cl> - Commit ID / Changelist"));
+	UE_LOG(LogZenOplogDiffTool, Display, TEXT("-diff-oplog -oplog1=<path> -oplog2=<path> - Diff 2 oplogs and output the results to a file/stdout"));
+	UE_LOG(LogZenOplogDiffTool, Display, TEXT("\t-results-path=<path> - Output the full results to a file"));
+	UE_LOG(LogZenOplogDiffTool, Display, TEXT("\t-output-values - Outputs the values of any property that was changed (optional)"));
+	UE_LOG(LogZenOplogDiffTool, Display, TEXT("\t-log-output=<option> - Optional logging to stdout"));
+	UE_LOG(LogZenOplogDiffTool, Display, TEXT("\t\tsummary - Outputs counts of identical, changed, and missing oplogs"));
+}
+
+void CreateFindSnapshotCommands(TQueue<TUniquePtr<IOplogDiffOperation>>& Operations)
+{
+	FString ProjectName;
+	FString PlatformName;
+	FString StreamName;
+	FString CommitID;
+	if (FParse::Value(FCommandLine::Get(), TEXT("-project="), ProjectName) &&
+		FParse::Value(FCommandLine::Get(), TEXT("-stream="), StreamName) &&
+		FParse::Value(FCommandLine::Get(), TEXT("-platform="), PlatformName) &&
+		FParse::Value(FCommandLine::Get(), TEXT("-commit="), CommitID))
+	{
+		TUniquePtr<FFindSnapshot> Operation = MakeUnique<FFindSnapshot>();
+		Operation->FindPlatform = PlatformName;
+		Operation->FindProjectName = ProjectName;
+		Operation->FindStream = StreamName;
+		Operation->FindCommitID = CommitID;
+		Operation->bDisplayAllBuildsFound = FParse::Param(FCommandLine::Get(), TEXT("-displaybuilds"));
+		Operations.Enqueue(MoveTemp(Operation));
+	}
+}
+
+FOutputDiffResults::ELogOptions GetDiffLogOptionsFromCmdLine()
+{
+	FString LogOptionString;
+	FOutputDiffResults::ELogOptions Logging = FOutputDiffResults::ELogOptions::None;
+	FParse::Value(FCommandLine::Get(), TEXT("-log-output="), LogOptionString);
+	if (LogOptionString == "summary")
+	{
+		Logging = FOutputDiffResults::ELogOptions::Summary;
+	}
+	return Logging;
+}
+
+FOutputDiffResults::EOutputOptions GetDiffOutputOptionsFromCmdLine()
+{
+	FOutputDiffResults::EOutputOptions OutputOption = FOutputDiffResults::EOutputOptions::None;
+	if (FParse::Param(FCommandLine::Get(), TEXT("-output-values")))
+	{
+		OutputOption = FOutputDiffResults::EOutputOptions::OutputChangedValues;
+	}
+	return OutputOption;
+}
+
+void CreateManifestDiffCommands(TQueue<TUniquePtr<IOplogDiffOperation>>& Operations)
+{
+	FString Manifest1, Manifest2;
+	if (FParse::Value(FCommandLine::Get(), TEXT("-oplog1="), Manifest1) && FParse::Value(FCommandLine::Get(), TEXT("-oplog2="), Manifest2))
+	{
+		TUniquePtr<IOplogDiffOperation> DiffOff = MakeUnique<FOplogManifestDiff>(Manifest1, Manifest2);
+		DiffOff->OnComplete = [&Operations](IOplogDiffOperation& CompletedDiff)
+		{
+			FString ResultsPath;
+			FParse::Value(FCommandLine::Get(), TEXT("-results-path="), ResultsPath);
+			FOplogManifestDiff& DiffOp = static_cast<FOplogManifestDiff&>(CompletedDiff);
+			TUniquePtr<FOutputDiffResults> OutputOp = MakeUnique<FOutputDiffResults>(MoveTemp(DiffOp.Manifest1),
+				MoveTemp(DiffOp.Manifest2),
+				MoveTemp(DiffOp.DiffResults));
+			OutputOp->OutputPath = ResultsPath;
+			OutputOp->LoggerOptions = GetDiffLogOptionsFromCmdLine();
+			OutputOp->OutputOptions = GetDiffOutputOptionsFromCmdLine();
+			Operations.Enqueue(MoveTemp(OutputOp));
+		};
+		Operations.Enqueue(MoveTemp(DiffOff));
+	}
+}
+
+void QueueDownloadOperations(TQueue<TUniquePtr<IOplogDiffOperation>>& Operations, const TArray<UE::FZenSnapshotDescriptor>& Descriptors, bool bDownloadAsJson, FStringView OutputDirectory)
+{
+	for (const UE::FZenSnapshotDescriptor& Snapshot : Descriptors)
+	{
+		Operations.Enqueue(MakeUnique<FDownloadSnapshot>(Snapshot, bDownloadAsJson, OutputDirectory));
+	}
+}
+
+void CreateDownloadCommands(TQueue<TUniquePtr<IOplogDiffOperation>>& Operations)
+{
+	FString SnapshotJsonPath;
+	FString OutputDirectory;
+	FString ProjectName;
+	FString PlatformName;
+	FString StreamName;
+	FString CommitID;
+
+	bool bDownloadAsJson = FParse::Param(FCommandLine::Get(), TEXT("-as-json"));
+	FParse::Value(FCommandLine::Get(), TEXT("-output-directory="), OutputDirectory);
+	if (FParse::Value(FCommandLine::Get(), TEXT("-snapshot-json="), SnapshotJsonPath))
+	{
+		// Load the descriptor, then queue the downloads
+		TUniquePtr<FLoadSnapshotDescriptor> LoadSnapshot = MakeUnique<FLoadSnapshotDescriptor>(SnapshotJsonPath);
+		LoadSnapshot->OnComplete = [&Operations, bDownloadAsJson, OutputDirectory](IOplogDiffOperation& CompletedLoad)
+		{
+			FLoadSnapshotDescriptor& LoadSnapshotOp = static_cast<FLoadSnapshotDescriptor&>(CompletedLoad);
+			QueueDownloadOperations(Operations, LoadSnapshotOp.LoadedSnapshots, bDownloadAsJson, OutputDirectory);
+		};
+		Operations.Enqueue(MoveTemp(LoadSnapshot));
+	}
+	else if (FParse::Value(FCommandLine::Get(), TEXT("-project="), ProjectName) &&
+		FParse::Value(FCommandLine::Get(), TEXT("-stream="), StreamName) &&
+		FParse::Value(FCommandLine::Get(), TEXT("-platform="), PlatformName) &&
+		FParse::Value(FCommandLine::Get(), TEXT("-commit="), CommitID))
+	{
+		// Find the snapshot, then queue the download
+		TUniquePtr<FFindSnapshot> FindSnapshot = MakeUnique<FFindSnapshot>();
+		FindSnapshot->FindPlatform = PlatformName;
+		FindSnapshot->FindProjectName = ProjectName;
+		FindSnapshot->FindStream = StreamName;
+		FindSnapshot->FindCommitID = CommitID;
+		FindSnapshot->OnComplete = [&Operations, bDownloadAsJson, OutputDirectory](IOplogDiffOperation& CompletedLoad)
+		{
+			FFindSnapshot& FindSnapshotOp = static_cast<FFindSnapshot&>(CompletedLoad);
+			TArray<UE::FZenSnapshotDescriptor> Descriptors;
+			Descriptors.Add(*FindSnapshotOp.FoundSnapshot);
+			QueueDownloadOperations(Operations, Descriptors, bDownloadAsJson, OutputDirectory);
+		};
+		Operations.Enqueue(MoveTemp(FindSnapshot));
+	}
+}
+
+void CreateInitialCommands(TQueue<TUniquePtr<IOplogDiffOperation>>& Operations)
+{
+	if (FParse::Param(FCommandLine::Get(), TEXT("-find-snapshot")))
+	{
+		CreateFindSnapshotCommands(Operations);
+	}
+	else if (FParse::Param(FCommandLine::Get(), TEXT("-download-snapshot")))
+	{
+		CreateDownloadCommands(Operations);
+	}
+	else if (FParse::Param(FCommandLine::Get(), TEXT("-diff-oplog")))
+	{
+		CreateManifestDiffCommands(Operations);
+	}
+}
+
+INT32_MAIN_INT32_ARGC_TCHAR_ARGV()
+{
+	const FTaskTagScope Scope(ETaskTag::EGameThread);
+
+	if (int32 Ret = GEngineLoop.PreInit(ArgC, ArgV))
+	{
+		return Ret;
+	}
+
+	ON_SCOPE_EXIT
+	{
+		RequestEngineExit("Exiting");
+		FEngineLoop::AppPreExit();
+		FModuleManager::Get().UnloadModulesAtShutdown();
+		FEngineLoop::AppExit();
+	};
+
+	// Don't allow more than one instance to run simultaneously 
+	FSystemWideCriticalSection SystemWideCritSection(TEXT("ZenOplogDiffTool"));
+	if (!SystemWideCritSection.IsValid())
+	{
+		UE_LOG(LogZenOplogDiffTool, Warning, TEXT("Zen Oplog Diff is already running!"));
+		return 1;
+	}
+
+	// Disable display of log timestamps, categories and verbosity
+	GLog->Flush();
+	TGuardValue NoTime(GPrintLogTimes, ELogTimes::None);
+	TGuardValue NoCategory(GPrintLogCategory, false);
+	TGuardValue NoVerbosity(GPrintLogVerbosity, false);
+
+	TQueue<TUniquePtr<IOplogDiffOperation>> AllOperations;	// Operations can be queued from other operations during the main loop!
+	CreateInitialCommands(AllOperations);
+
+	int32 ExitCode = 0;
+	if (AllOperations.IsEmpty())
+	{
+		PrintHelp();
+	}
+	else
+	{
+		// Run a light version of the engine main loop
+		TUniquePtr<IOplogDiffOperation> RunningOperation;
+		double DeltaTime = 0.0;
+		double LastTime = FPlatformTime::Seconds();
+		while (!IsEngineExitRequested())
+		{
+			BeginExitIfRequested();
+
+			FTaskGraphInterface::Get().ProcessThreadUntilIdle(ENamedThreads::GameThread);
+			FTSTicker::GetCoreTicker().Tick(DeltaTime);
+
+			if (RunningOperation == nullptr && !AllOperations.Dequeue(RunningOperation))
+			{
+				break;
+			}
+			else if (RunningOperation)
+			{
+				IOplogDiffOperation::ERunningState CurrentState = RunningOperation->Run();
+				if (CurrentState != IOplogDiffOperation::ERunningState::Running)
+				{
+					if (CurrentState == IOplogDiffOperation::ERunningState::Success && RunningOperation->OnComplete)
+					{
+						RunningOperation->OnComplete(*RunningOperation);
+					}
+					else if (CurrentState == IOplogDiffOperation::ERunningState::Error)
+					{
+						ExitCode = 1;
+						break;
+					}
+					RunningOperation = nullptr;
+				}
+			}
+
+			DeltaTime = FPlatformTime::Seconds() - LastTime;
+			LastTime = FPlatformTime::Seconds();
+
+			UE::Stats::FStats::AdvanceFrame(false);
+			FCoreDelegates::OnEndFrame.Broadcast();
+
+			LLM(FLowLevelMemTracker::Get().UpdateStatsPerFrame());
+
+			GFrameCounter++;
+		}
+	}
+
+	return ExitCode;
+}
diff --git a/Engine/Source/Programs/ZenOplogDiffTool/ZenOplogDiffTool.Build.cs b/Engine/Source/Programs/ZenOplogDiffTool/ZenOplogDiffTool.Build.cs
new file mode 100644
--- /dev/null
+++ b/Engine/Source/Programs/ZenOplogDiffTool/ZenOplogDiffTool.Build.cs
@@ -0,0 +1,39 @@
+// Copyright Epic Games, Inc. All Rights Reserved.
+
+using UnrealBuildTool;
+
+public class ZenOplogDiffTool : ModuleRules
+{
+	public ZenOplogDiffTool(ReadOnlyTargetRules Target) : base(Target)
+	{
+		PublicIncludePathModuleNames.Add("Launch");
+		PrivateDependencyModuleNames.Add("ApplicationCore");
+		PrivateDependencyModuleNames.Add("Core");
+		PrivateDependencyModuleNames.Add("DesktopPlatform");
+		PrivateDependencyModuleNames.Add("Projects");
+		PrivateDependencyModuleNames.Add("ZenOplogUtils");
+
+		if (Target.Platform == UnrealTargetPlatform.Win64)
+		{
+			RuntimeDependencies.Add("$(EngineDir)/Binaries/Win64/crashpad_handler.exe");
+			RuntimeDependencies.Add("$(EngineDir)/Binaries/Win64/zen.exe");
+			RuntimeDependencies.Add("$(EngineDir)/Binaries/Win64/zen.pdb");
+			RuntimeDependencies.Add("$(EngineDir)/Binaries/Win64/zenserver.exe");
+			RuntimeDependencies.Add("$(EngineDir)/Binaries/Win64/zenserver.pdb");
+			RuntimeDependencies.Add("$(EngineDir)/Binaries/DotNET/OidcToken/win-x64/OidcToken.exe");
+		}
+		else if (Target.Platform == UnrealTargetPlatform.Linux)
+		{
+			RuntimeDependencies.Add("$(EngineDir)/Binaries/Linux/zen");
+			RuntimeDependencies.Add("$(EngineDir)/Binaries/Linux/zenserver");
+			RuntimeDependencies.Add("$(EngineDir)/Binaries/DotNET/OidcToken/linux-x64/OidcToken");
+		}
+		else if (Target.Platform == UnrealTargetPlatform.Mac)
+		{
+			RuntimeDependencies.Add("$(EngineDir)/Binaries/Mac/crashpad_handler");
+			RuntimeDependencies.Add("$(EngineDir)/Binaries/Mac/zen");
+			RuntimeDependencies.Add("$(EngineDir)/Binaries/Mac/zenserver");
+			RuntimeDependencies.Add("$(EngineDir)/Binaries/DotNET/OidcToken/osx-x64/OidcToken");
+		}
+	}
+}
diff --git a/Engine/Source/Programs/ZenOplogDiffTool/ZenOplogDiffTool.Target.cs b/Engine/Source/Programs/ZenOplogDiffTool/ZenOplogDiffTool.Target.cs
new file mode 100644
--- /dev/null
+++ b/Engine/Source/Programs/ZenOplogDiffTool/ZenOplogDiffTool.Target.cs
@@ -0,0 +1,32 @@
+// Copyright Epic Games, Inc. All Rights Reserved.
+
+using UnrealBuildTool;
+
+[SupportedPlatforms(UnrealPlatformClass.Desktop)]
+public class ZenOplogDiffToolTarget : TargetRules
+{
+	public ZenOplogDiffToolTarget(TargetInfo Target) : base(Target)
+	{
+		Type = TargetType.Program;
+		IncludeOrderVersion = EngineIncludeOrderVersion.Latest;
+		LinkType = TargetLinkType.Modular;
+		LaunchModuleName = "ZenOplogDiffTool";
+
+		bBuildDeveloperTools = false;
+		bCompileAgainstApplicationCore = true;
+		bCompileAgainstCoreUObject = false;
+		bCompileAgainstEngine = false;
+		bCompileICU = false;
+		bForceDisableAutomationTests = true;
+		bIsBuildingConsoleApplication = true;
+		bUseLoggingInShipping = true;
+
+		if (Target.Configuration == UnrealTargetConfiguration.Debug ||
+			Target.Configuration == UnrealTargetConfiguration.Development)
+		{
+			GlobalDefinitions.Add("LLM_ENABLED_IN_CONFIG=1");
+			GlobalDefinitions.Add("UE_MEMORY_TAGS_TRACE_ENABLED=1");
+			bEnableTrace = true;
+		}
+	}
+}
